Renamed SFilter and SFilterList classes in LogVisualizer to solve linking conflicts on Mac

diff --git a/Engine/Source/Developer/LogVisualizer/Private/SFilterList.cpp b/Engine/Source/Developer/LogVisualizer/Private/SFilterList.cpp
--- a/Engine/Source/Developer/LogVisualizer/Private/SFilterList.cpp
+++ b/Engine/Source/Developer/LogVisualizer/Private/SFilterList.cpp
@@ -2,7 +2,7 @@
 #include "SFilterList.h"
 
 /////////////////////
-// SFilterList
+// SLogFilterList
 /////////////////////
 
 #define LOCTEXT_NAMESPACE "SLogVisualizer"
@@ -53,15 +53,15 @@ private:
 /**
 * A single filter in the filter list. Can be removed by clicking the remove button on it.
 */
-class SFilter : public SCompoundWidget
+class SLogFilter : public SCompoundWidget
 {
 public:
-	DECLARE_DELEGATE_OneParam(FOnRequestRemove, const TSharedRef<SFilter>& /*FilterToRemove*/);
-	DECLARE_DELEGATE_OneParam(FOnRequestEnableOnly, const TSharedRef<SFilter>& /*FilterToEnable*/);
+	DECLARE_DELEGATE_OneParam(FOnRequestRemove, const TSharedRef<SLogFilter>& /*FilterToRemove*/);
+	DECLARE_DELEGATE_OneParam(FOnRequestEnableOnly, const TSharedRef<SLogFilter>& /*FilterToEnable*/);
 	DECLARE_DELEGATE(FOnRequestDisableAll);
 	DECLARE_DELEGATE(FOnRequestRemoveAll);
 
-	SLATE_BEGIN_ARGS(SFilter){}
+	SLATE_BEGIN_ARGS(SLogFilter){}
 
 		/** If this is an front end filter, this is the filter object */
 		SLATE_ARGUMENT(FName, FilterName)
@@ -69,7 +69,7 @@ public:
 		SLATE_ARGUMENT(FLinearColor, ColorCategory)
 
 		/** Invoked when the filter toggled */
-		SLATE_EVENT(SFilterList::FOnFilterChanged, OnFilterChanged)
+		SLATE_EVENT(SLogFilterList::FOnFilterChanged, OnFilterChanged)
 
 		/** Invoked when a request to remove this filter originated from within this filter */
 		SLATE_EVENT(FOnRequestRemove, OnRequestRemove)
@@ -111,14 +111,14 @@ public:
 						SAssignNew(ToggleButtonPtr, SFilterCheckBox)
 						.Style(FEditorStyle::Get(), "ContentBrowser.FilterButton")
 						.ToolTipText(FilterToolTip)
-						.Padding(this, &SFilter::GetFilterNamePadding)
-						.IsChecked(this, &SFilter::IsChecked)
-						.OnCheckStateChanged(this, &SFilter::FilterToggled)
-						/*.OnGetMenuContent(this, &SFilter::GetRightClickMenuContent)*/
-						.ForegroundColor(this, &SFilter::GetFilterForegroundColor)
+						.Padding(this, &SLogFilter::GetFilterNamePadding)
+						.IsChecked(this, &SLogFilter::IsChecked)
+						.OnCheckStateChanged(this, &SLogFilter::FilterToggled)
+						/*.OnGetMenuContent(this, &SLogFilter::GetRightClickMenuContent)*/
+						.ForegroundColor(this, &SLogFilter::GetFilterForegroundColor)
 						[
 							SNew(STextBlock)
-							.ColorAndOpacity(this, &SFilter::GetFilterNameColorAndOpacity)
+							.ColorAndOpacity(this, &SLogFilter::GetFilterNameColorAndOpacity)
 							.Font(FEditorStyle::GetFontStyle("ContentBrowser.FilterNameFont"))
 							.ShadowOffset(FVector2D(1.f, 1.f))
 							.Text(GetFilterNameAsString().Replace(TEXT("Log"), TEXT(""), ESearchCase::CaseSensitive))
@@ -126,8 +126,8 @@ public:
 					]
 				];
 
-			ToggleButtonPtr->SetOnFilterDoubleClicked(FOnClicked::CreateSP(this, &SFilter::FilterDoubleClicked));
-			ToggleButtonPtr->SetOnFilterMiddleButtonClicked(FOnClicked::CreateSP(this, &SFilter::FilterMiddleButtonClicked));
+			ToggleButtonPtr->SetOnFilterDoubleClicked(FOnClicked::CreateSP(this, &SLogFilter::FilterDoubleClicked));
+			ToggleButtonPtr->SetOnFilterMiddleButtonClicked(FOnClicked::CreateSP(this, &SLogFilter::FilterMiddleButtonClicked));
 		}
 
 	/** Sets whether or not this filter is applied to the combined filter */
@@ -200,14 +200,14 @@ private:
 				FText::Format(LOCTEXT("RemoveFilter", "Remove: {0}"), FiletNameAsText),
 				LOCTEXT("RemoveFilterTooltip", "Remove this filter from the list. It can be added again in the filters menu."),
 				FSlateIcon(),
-				FUIAction(FExecuteAction::CreateSP(this, &SFilter::RemoveFilter))
+				FUIAction(FExecuteAction::CreateSP(this, &SLogFilter::RemoveFilter))
 				);
 
 			MenuBuilder.AddMenuEntry(
 				FText::Format(LOCTEXT("EnableOnlyThisFilter", "Enable this only: {0}"), FiletNameAsText),
 				LOCTEXT("EnableOnlyThisFilterTooltip", "Enable only this filter from the list."),
 				FSlateIcon(),
-				FUIAction(FExecuteAction::CreateSP(this, &SFilter::EnableOnly))
+				FUIAction(FExecuteAction::CreateSP(this, &SLogFilter::EnableOnly))
 				);
 
 		}
@@ -219,14 +219,14 @@ private:
 				LOCTEXT("DisableAllFilters", "Disable All Filters"),
 				LOCTEXT("DisableAllFiltersTooltip", "Disables all active filters."),
 				FSlateIcon(),
-				FUIAction(FExecuteAction::CreateSP(this, &SFilter::DisableAllFilters))
+				FUIAction(FExecuteAction::CreateSP(this, &SLogFilter::DisableAllFilters))
 				);
 
 			MenuBuilder.AddMenuEntry(
 				LOCTEXT("RemoveAllFilters", "Remove All Filters"),
 				LOCTEXT("RemoveAllFiltersTooltip", "Removes all filters from the list."),
 				FSlateIcon(),
-				FUIAction(FExecuteAction::CreateSP(this, &SFilter::RemoveAllFilters))
+				FUIAction(FExecuteAction::CreateSP(this, &SLogFilter::RemoveAllFilters))
 				);
 		}
 		MenuBuilder.EndSection();
@@ -237,14 +237,14 @@ private:
 	/** Removes this filter from the filter list */
 	void RemoveFilter()
 	{
-		TSharedRef<SFilter> Self = SharedThis(this);
+		TSharedRef<SLogFilter> Self = SharedThis(this);
 		OnRequestRemove.ExecuteIfBound(Self);
 	}
 
 	/** Enables only this filter from the filter list */
 	void EnableOnly()
 	{
-		TSharedRef<SFilter> Self = SharedThis(this);
+		TSharedRef<SLogFilter> Self = SharedThis(this);
 		OnRequestEnableOnly.ExecuteIfBound(Self);
 	}
 
@@ -287,7 +287,7 @@ private:
 
 private:
 	/** Invoked when the filter toggled */
-	SFilterList::FOnFilterChanged OnFilterChanged;
+	SLogFilterList::FOnFilterChanged OnFilterChanged;
 
 	/** Invoked when a request to remove this filter originated from within this filter */
 	FOnRequestRemove OnRequestRemove;
@@ -323,7 +323,7 @@ private:
 
 
 
-void SFilterList::Construct(const FArguments& InArgs)
+void SLogFilterList::Construct(const FArguments& InArgs)
 {
 	OnGetContextMenu = InArgs._OnGetContextMenu;
 	OnFilterChanged = InArgs._OnFilterChanged;
@@ -335,21 +335,21 @@ void SFilterList::Construct(const FArguments& InArgs)
 		];
 }
 
-void SFilterList::SomeFilterGetChanged()
+void SLogFilterList::SomeFilterGetChanged()
 {
 	OnFilterChanged.ExecuteIfBound();
 }
 
-void SFilterList::AddFilter(const FString& InFilterName, FLinearColor InColorCategory)
+void SLogFilterList::AddFilter(const FString& InFilterName, FLinearColor InColorCategory)
 {
-	TSharedRef<SFilter> NewFilter =
-		SNew(SFilter)
+	TSharedRef<SLogFilter> NewFilter =
+		SNew(SLogFilter)
 		.FilterName(*InFilterName)
 		.ColorCategory(InColorCategory)
-		.OnFilterChanged(this, &SFilterList::SomeFilterGetChanged)
-		//.OnRequestRemove(this, &SFilterList::RemoveFilter)
-		//.OnRequestDisableAll(this, &SFilterList::DisableAllFilters)
-		//.OnRequestRemoveAll(this, &SFilterList::RemoveAllFilters);
+		.OnFilterChanged(this, &SLogFilterList::SomeFilterGetChanged)
+		//.OnRequestRemove(this, &SLogFilterList::RemoveFilter)
+		//.OnRequestDisableAll(this, &SLogFilterList::DisableAllFilters)
+		//.OnRequestRemoveAll(this, &SLogFilterList::RemoveAllFilters);
 		;
 	NewFilter->SetEnabled(true);
 	Filters.Add(NewFilter);
@@ -361,12 +361,12 @@ void SFilterList::AddFilter(const FString& InFilterName, FLinearColor InColorCat
 		];
 }
 
-bool SFilterList::IsFilterEnabled(const FString& InFilterName, TEnumAsByte<ELogVerbosity::Type> Verbosity)
+bool SLogFilterList::IsFilterEnabled(const FString& InFilterName, TEnumAsByte<ELogVerbosity::Type> Verbosity)
 {
 	const FName FilterName(*InFilterName);
 	for (int32 Index = 0; Index < Filters.Num(); ++Index)
 	{
-		const SFilter& Filter = Filters[Index].Get();
+		const SLogFilter& Filter = Filters[Index].Get();
 		if (Filter.GetFilterName() == FilterName)
 		{
 			return Filter.IsEnabled();
diff --git a/Engine/Source/Developer/LogVisualizer/Private/SFilterList.h b/Engine/Source/Developer/LogVisualizer/Private/SFilterList.h
--- a/Engine/Source/Developer/LogVisualizer/Private/SFilterList.h
+++ b/Engine/Source/Developer/LogVisualizer/Private/SFilterList.h
@@ -2,12 +2,12 @@
 
 #pragma once
 
-class SFilter;
+class SLogFilter;
 
 /**
 * A list of filters currently applied to an asset view.
 */
-class SFilterList : public SCompoundWidget
+class SLogFilterList : public SCompoundWidget
 {
 public:
 	/** Delegate for when filters have changed */
@@ -15,7 +15,7 @@ public:
 
 	DECLARE_DELEGATE_RetVal(TSharedPtr<SWidget>, FOnGetContextMenu);
 
-	SLATE_BEGIN_ARGS(SFilterList){}
+	SLATE_BEGIN_ARGS(SLogFilterList){}
 
 		/** Called when an asset is right clicked */
 		SLATE_EVENT(FOnGetContextMenu, OnGetContextMenu)
@@ -34,7 +34,7 @@ public:
 
 	void SomeFilterGetChanged();
 
-	TArray<TSharedRef<SFilter> > Filters;
+	TArray<TSharedRef<SLogFilter> > Filters;
 
 	/** The horizontal box which contains all the filters */
 	TSharedPtr<SWrapBox> FilterBox;
diff --git a/Engine/Source/Developer/LogVisualizer/Private/SLogVisualizer.cpp b/Engine/Source/Developer/LogVisualizer/Private/SLogVisualizer.cpp
--- a/Engine/Source/Developer/LogVisualizer/Private/SLogVisualizer.cpp
+++ b/Engine/Source/Developer/LogVisualizer/Private/SLogVisualizer.cpp
@@ -204,7 +204,7 @@ void SLogVisualizer::Construct(const FArguments& InArgs, FLogVisualizer* InLogVi
 			[
 				SNew(STutorialWrapper, TEXT("CategoryFilters"))
 				[
-					SAssignNew(FilterListPtr, SFilterList)
+					SAssignNew(FilterListPtr, SLogFilterList)
 					.OnFilterChanged(this, &SLogVisualizer::OnLogCategoryFiltersChanged)
 					/*.OnGetContextMenu(this, &SLogVisualizer::GetFilterContextMenu)*/
 					/*.FrontendFilters(FrontendFilters)*/
diff --git a/Engine/Source/Developer/LogVisualizer/Private/SLogVisualizer.h b/Engine/Source/Developer/LogVisualizer/Private/SLogVisualizer.h
--- a/Engine/Source/Developer/LogVisualizer/Private/SLogVisualizer.h
+++ b/Engine/Source/Developer/LogVisualizer/Private/SLogVisualizer.h
@@ -5,7 +5,7 @@
 #include "STimeline.h"
 #include "Debug/DebugDrawService.h"
 
-class SFilterList;
+class SLogFilterList;
 
 struct FLogsListItem
 {
@@ -214,7 +214,7 @@ private:
 	// MEMBERS
 
 	/** The filter list */
-	TSharedPtr<SFilterList> FilterListPtr;
+	TSharedPtr<SLogFilterList> FilterListPtr;
 
 	/** Index into LogVisualizer->Logs array for entries you want to show. */
 	TArray<TSharedPtr<FLogsListItem> >	LogsList;
