Added a little sanity check preventing LogVisualizers from crashing when deleting vlogs #UE4

diff --git a/Engine/Source/Developer/LogVisualizer/Private/SLogVisualizer.cpp b/Engine/Source/Developer/LogVisualizer/Private/SLogVisualizer.cpp
--- a/Engine/Source/Developer/LogVisualizer/Private/SLogVisualizer.cpp
+++ b/Engine/Source/Developer/LogVisualizer/Private/SLogVisualizer.cpp
@@ -524,7 +524,7 @@ void SLogVisualizer::OnLogCategoryFiltersChanged()
 	if (LogVisualizer && LogVisualizer->Logs.IsValidIndex(SelectedLogIndex))
 	{
 		TSharedPtr<FActorsVisLog> Log = LogVisualizer->Logs[SelectedLogIndex];
-		if (Log.IsValid())
+		if (Log.IsValid() && Log->Entries.IsValidIndex(LogEntryIndex))
 		{
 			ShowEntry(Log->Entries[LogEntryIndex].Get());
 		}
