Fixed some coding style offenders in LogVisualizer module #UE4

diff --git a/Engine/Source/Developer/LogVisualizer/Private/SLogVisualizer.cpp b/Engine/Source/Developer/LogVisualizer/Private/SLogVisualizer.cpp
--- a/Engine/Source/Developer/LogVisualizer/Private/SLogVisualizer.cpp
+++ b/Engine/Source/Developer/LogVisualizer/Private/SLogVisualizer.cpp
@@ -497,9 +497,8 @@ void SLogVisualizer::OnListDoubleClick(TSharedPtr<FLogsListItem> LogListItem)
 	UEditorEngine *EEngine = Cast<UEditorEngine>(GEngine);
 	if (GIsEditor && EEngine != NULL)
 	{
-		for (int32 i = 0; i < EEngine->AllViewportClients.Num(); i++)
+		for (auto ViewportClient : EEngine->AllViewportClients)
 		{
-			FEditorViewportClient* ViewportClient = EEngine->AllViewportClients[i];
 			//GEditor->MoveViewportCamerasToActor(SelectedActors, bActiveViewportOnly);
 			ViewportClient->FocusViewportOnBox(FBox::BuildAABB(Orgin, Extent));
 		}
@@ -533,18 +532,18 @@ void SLogVisualizer::GetVisibleEntries(const TSharedPtr<FActorsVisLog>& Log, TAr
 
 	if (FilterListPtr.IsValid())
 	{
-		for (int32 i = 0; i < Log->Entries.Num(); ++i)
+		for (int32 EntryIndex = 0; EntryIndex < Log->Entries.Num(); ++EntryIndex)
 		{
 			// if any log line is visible - add this entry
 			bool bAddedEntry = false;
 
 			if (!bAddedEntry)
 			{
-				for (int32 j = 0; j < Log->Entries[i]->LogLines.Num(); ++j)
+				for (int32 LogLineIndex = 0; LogLineIndex < Log->Entries[EntryIndex]->LogLines.Num(); ++LogLineIndex)
 				{
-					if (FilterListPtr->IsFilterEnabled(Log->Entries[i]->LogLines[j].Category.ToString(), Log->Entries[i]->LogLines[j].Verbosity))
+					if (FilterListPtr->IsFilterEnabled(Log->Entries[EntryIndex]->LogLines[LogLineIndex].Category.ToString(), Log->Entries[EntryIndex]->LogLines[LogLineIndex].Verbosity))
 					{
-						OutEntries.AddUnique(Log->Entries[i]);
+						OutEntries.AddUnique(Log->Entries[EntryIndex]);
 						bAddedEntry = true;
 						break;
 					}
@@ -556,11 +555,11 @@ void SLogVisualizer::GetVisibleEntries(const TSharedPtr<FActorsVisLog>& Log, TAr
 				continue;
 			}
 
-			for (int32 j = 0; j < Log->Entries[i]->ElementsToDraw.Num(); ++j)
+			for (int32 ElementIndex = 0; ElementIndex < Log->Entries[EntryIndex]->ElementsToDraw.Num(); ++ElementIndex)
 			{
-				if (Log->Entries[i]->ElementsToDraw[j].Category == NAME_None || FilterListPtr->IsFilterEnabled(Log->Entries[i]->ElementsToDraw[j].Category.ToString(), Log->Entries[i]->ElementsToDraw[j].Verbosity))
+				if (Log->Entries[EntryIndex]->ElementsToDraw[ElementIndex].Category == NAME_None || FilterListPtr->IsFilterEnabled(Log->Entries[EntryIndex]->ElementsToDraw[ElementIndex].Category.ToString(), Log->Entries[EntryIndex]->ElementsToDraw[ElementIndex].Verbosity))
 				{
-					OutEntries.AddUnique(Log->Entries[i]);
+					OutEntries.AddUnique(Log->Entries[EntryIndex]);
 					bAddedEntry = true;
 					break;
 				}
@@ -570,17 +569,17 @@ void SLogVisualizer::GetVisibleEntries(const TSharedPtr<FActorsVisLog>& Log, TAr
 				continue;
 			}
 
-			for (int32 SampleIndex = 0; SampleIndex < Log->Entries[i]->HistogramSamples.Num(); ++SampleIndex)
+			for (int32 SampleIndex = 0; SampleIndex < Log->Entries[EntryIndex]->HistogramSamples.Num(); ++SampleIndex)
 			{
-				const FName CurrentCategory = Log->Entries[i]->HistogramSamples[SampleIndex].Category;
-				const FName CurrentGraphName = Log->Entries[i]->HistogramSamples[SampleIndex].GraphName;
-				const FName CurrentDataName = Log->Entries[i]->HistogramSamples[SampleIndex].DataName;
+				const FName CurrentCategory = Log->Entries[EntryIndex]->HistogramSamples[SampleIndex].Category;
+				const FName CurrentGraphName = Log->Entries[EntryIndex]->HistogramSamples[SampleIndex].GraphName;
+				const FName CurrentDataName = Log->Entries[EntryIndex]->HistogramSamples[SampleIndex].DataName;
 
 				if (CurrentCategory == NAME_None || 
 					(FilterListPtr->IsFilterEnabled(CurrentCategory.ToString(), ELogVerbosity::All) &&
 					FilterListPtr->IsFilterEnabled(CurrentGraphName.ToString(), CurrentDataName.ToString(), ELogVerbosity::All)))
 				{
-					OutEntries.AddUnique(Log->Entries[i]);
+					OutEntries.AddUnique(Log->Entries[EntryIndex]);
 					bAddedEntry = true;
 					break;
 				}
@@ -1097,12 +1096,12 @@ void SLogVisualizer::OnLogStatusGetChildren(TSharedPtr<FLogStatusItem> InItem, T
 void SLogVisualizer::UpdateStatusItems(const FVisLogEntry* LogEntry)
 {
 	TArray<FString> ExpandedCategories;
-	for (int32 i = 0; i < StatusItems.Num(); i++)
+	for (int32 ItemIndex = 0; ItemIndex < StatusItems.Num(); ItemIndex++)
 	{
-		const bool bIsExpanded = StatusItemsView->IsItemExpanded(StatusItems[i]);
+		const bool bIsExpanded = StatusItemsView->IsItemExpanded(StatusItems[ItemIndex]);
 		if (bIsExpanded)
 		{
-			ExpandedCategories.Add(StatusItems[i]->ItemText);
+			ExpandedCategories.Add(StatusItems[ItemIndex]->ItemText);
 		}
 	}
 
@@ -1113,18 +1112,18 @@ void SLogVisualizer::UpdateStatusItems(const FVisLogEntry* LogEntry)
 		FString TimestampDesc = FString::Printf(TEXT("%.2fs"), LogEntry->TimeStamp);
 		StatusItems.Add(MakeShareable(new FLogStatusItem(LOCTEXT("VisLogTimestamp","Time").ToString(), TimestampDesc)));
 
-		for (int32 iCategory = 0; iCategory < LogEntry->Status.Num(); iCategory++)
+		for (int32 CategoryIndex = 0; CategoryIndex < LogEntry->Status.Num(); CategoryIndex++)
 		{
-			if (LogEntry->Status[iCategory].Data.Num() <= 0)
+			if (LogEntry->Status[CategoryIndex].Data.Num() <= 0)
 			{
 				continue;
 			}
 
-			TSharedRef<FLogStatusItem> StatusItem = MakeShareable(new FLogStatusItem(LogEntry->Status[iCategory].Category));
-			for (int32 iLine = 0; iLine < LogEntry->Status[iCategory].Data.Num(); iLine++)
+			TSharedRef<FLogStatusItem> StatusItem = MakeShareable(new FLogStatusItem(LogEntry->Status[CategoryIndex].Category));
+			for (int32 LineIndex = 0; LineIndex < LogEntry->Status[CategoryIndex].Data.Num(); LineIndex++)
 			{
 				FString KeyDesc, ValueDesc;
-				const bool bHasValue = LogEntry->Status[iCategory].GetDesc(iLine, KeyDesc, ValueDesc);
+				const bool bHasValue = LogEntry->Status[CategoryIndex].GetDesc(LineIndex, KeyDesc, ValueDesc);
 				if (bHasValue)
 				{
 					StatusItem->Children.Add(MakeShareable(new FLogStatusItem(KeyDesc, ValueDesc)));
@@ -1137,13 +1136,13 @@ void SLogVisualizer::UpdateStatusItems(const FVisLogEntry* LogEntry)
 
 	StatusItemsView->RequestTreeRefresh();
 
-	for (int32 iItem = 0; iItem < StatusItems.Num(); iItem++)
+	for (int32 ItemIndex = 0; ItemIndex < StatusItems.Num(); ItemIndex++)
 	{
-		for (int32 i = 0; i < ExpandedCategories.Num(); i++)
+		for (const FString& Category : ExpandedCategories)
 		{
-			if (StatusItems[iItem]->ItemText == ExpandedCategories[i])
+			if (StatusItems[ItemIndex]->ItemText == Category)
 			{
-				StatusItemsView->SetItemExpansion(StatusItems[iItem], true);
+				StatusItemsView->SetItemExpansion(StatusItems[ItemIndex], true);
 				break;
 			}
 		}
