Fixed a few warnings missed during the integration.

diff --git a/Engine/Source/Developer/LogVisualizer/Private/SLogVisualizer.cpp b/Engine/Source/Developer/LogVisualizer/Private/SLogVisualizer.cpp
--- a/Engine/Source/Developer/LogVisualizer/Private/SLogVisualizer.cpp
+++ b/Engine/Source/Developer/LogVisualizer/Private/SLogVisualizer.cpp
@@ -44,7 +44,7 @@ template <typename ItemType>
 class SLogListView : public SListView<ItemType>
 {
 public:
-	virtual FReply OnMouseWheel(const FGeometry& MyGeometry, const FPointerEvent& MouseEvent) OVERRIDE
+	virtual FReply OnMouseWheel(const FGeometry& MyGeometry, const FPointerEvent& MouseEvent) override
 	{
 		if (!MouseEvent.IsLeftShiftDown())
 		{
diff --git a/Engine/Source/Editor/UnrealEd/Public/CreditsScreen.h b/Engine/Source/Editor/UnrealEd/Public/CreditsScreen.h
--- a/Engine/Source/Editor/UnrealEd/Public/CreditsScreen.h
+++ b/Engine/Source/Editor/UnrealEd/Public/CreditsScreen.h
@@ -17,7 +17,7 @@ public:
 	 */
 	UNREALED_API void Construct(const FArguments& InArgs);
 
-	UNREALED_API virtual void Tick(const FGeometry& AllottedGeometry, const double InCurrentTime, const float InDeltaTime) OVERRIDE;
+	UNREALED_API virtual void Tick(const FGeometry& AllottedGeometry, const double InCurrentTime, const float InDeltaTime) override;
 
 private:
 
