Fix CIS error on Mac.

diff --git a/Engine/Source/Developer/LogVisualizer/Private/LogVisualizer.h b/Engine/Source/Developer/LogVisualizer/Private/LogVisualizer.h
--- a/Engine/Source/Developer/LogVisualizer/Private/LogVisualizer.h
+++ b/Engine/Source/Developer/LogVisualizer/Private/LogVisualizer.h
@@ -4,6 +4,7 @@
 
 #if ENABLE_VISUAL_LOG
 
+struct FVisualLog;
 struct FActorsVisLog;
 
 /** Actual implementation of LogVisualizer, private inside module */
