Another CIS fix on Mac.

diff --git a/Engine/Source/Developer/LogVisualizer/Private/LogVisualizerModule.cpp b/Engine/Source/Developer/LogVisualizer/Private/LogVisualizerModule.cpp
--- a/Engine/Source/Developer/LogVisualizer/Private/LogVisualizerModule.cpp
+++ b/Engine/Source/Developer/LogVisualizer/Private/LogVisualizerModule.cpp
@@ -1,6 +1,7 @@
 // Copyright 1998-2014 Epic Games, Inc. All Rights Reserved.
 
 #include "LogVisualizerPCH.h"
+#include "LogVisualizer.h"
 
 IMPLEMENT_MODULE(FLogVisualizerModule, LogVisualizer);
 DEFINE_LOG_CATEGORY(LogLogVisualizer);
