Fixed removing ALogVisualizerDebugActor from scene after closing LogVisualizer window. #ue4

diff --git a/Engine/Source/Developer/LogVisualizer/Private/LogVisualizer.cpp b/Engine/Source/Developer/LogVisualizer/Private/LogVisualizer.cpp
--- a/Engine/Source/Developer/LogVisualizer/Private/LogVisualizer.cpp
+++ b/Engine/Source/Developer/LogVisualizer/Private/LogVisualizer.cpp
@@ -53,7 +53,6 @@ void FLogVisualizer::CloseUI(UWorld* InWorld)
 	{
 		if (LogWindow.IsValid() && (World.IsValid() == false || World == InWorld))
 		{
-			InWorld->DestroyActor(GetHelperActor(InWorld));
 			DebugActor = NULL;
 
 			CleanUp();
diff --git a/Engine/Source/Developer/LogVisualizer/Private/SLogVisualizer.cpp b/Engine/Source/Developer/LogVisualizer/Private/SLogVisualizer.cpp
--- a/Engine/Source/Developer/LogVisualizer/Private/SLogVisualizer.cpp
+++ b/Engine/Source/Developer/LogVisualizer/Private/SLogVisualizer.cpp
@@ -571,6 +571,11 @@ END_SLATE_FUNCTION_BUILD_OPTIMIZATION
 
 SLogVisualizer::~SLogVisualizer()
 {
+	UWorld* World = GetWorld();
+	if (World)
+	{
+		World->DestroyActor(FLogVisualizerModule::Get()->GetHelperActor(World));
+	}
 	UGameplayDebuggingComponent::OnDebuggingTargetChangedDelegate.RemoveAll(this);
 	LogVisualizer->OnLogAdded().RemoveAll(this);
 	UDebugDrawService::Unregister(DrawingOnCanvasDelegate);
