Tweaks for LogVisualizer reports #ue4

diff --git a/Engine/Source/Developer/LogVisualizer/Private/SLogVisualizerReport.cpp b/Engine/Source/Developer/LogVisualizer/Private/SLogVisualizerReport.cpp
--- a/Engine/Source/Developer/LogVisualizer/Private/SLogVisualizerReport.cpp
+++ b/Engine/Source/Developer/LogVisualizer/Private/SLogVisualizerReport.cpp
@@ -64,6 +64,8 @@ FText SLogVisualizerReport::GenerateReportText() const
 {
 	FString OutString;
 	TArray<FVisualLogEntry::FLogEvent> GlobalEventsStats;
+	TMap<FString, TArray<FString> >	 EventToObjectsMap;
+
 
 	for (auto LogItem : SelectedItems)
 	{
@@ -103,52 +105,57 @@ FText SLogVisualizerReport::GenerateReportText() const
 		}
 		for (auto& CurrentEvent : AllEvents)
 		{
-			bPrintNextLine = CurrentEvent.EventTags.Num() > 0;
 			for (auto& CurrentEventTag : CurrentEvent.EventTags)
 			{
-				OutString.Append(FString::Printf(TEXT("        %s <%s> occurred %d times\n"), *CurrentEvent.Name, *CurrentEventTag.Key.ToString(), CurrentEventTag.Value));
+				OutString.Append(FString::Printf(TEXT("        %s  with <%s> tag occurred    %d times\n"), *CurrentEvent.Name, *CurrentEventTag.Key.ToString(), CurrentEventTag.Value));
 			}
-			OutString.Append( FString::Printf(TEXT("        %s occurred %d times\n"), *CurrentEvent.Name, CurrentEvent.Counter) );
+			OutString.Append( FString::Printf(TEXT("        %s occurred %d times\n\n"), *CurrentEvent.Name, CurrentEvent.Counter) );
 
+			bool bJustAdded = false;
 			int32 Index = GlobalEventsStats.Find(FVisualLogEntry::FLogEvent(CurrentEvent));
 			if (Index != INDEX_NONE)
 			{
 				GlobalEventsStats[Index].Counter += CurrentEvent.Counter;
-				for (auto& CurrentEventTag : CurrentEvent.EventTags)
-				{
-					if (GlobalEventsStats[Index].EventTags.Contains(CurrentEventTag.Key))
-					{
-						GlobalEventsStats[Index].EventTags[CurrentEventTag.Key] += CurrentEventTag.Value;
-					}
-					else
-					{
-						GlobalEventsStats[Index].EventTags.Add(CurrentEventTag.Key, CurrentEventTag.Value);
-					}
-				}
+				EventToObjectsMap[CurrentEvent.Name] .AddUnique(LogItem->Name.ToString());
 			}
 			else
 			{
+				bJustAdded = true;
 				GlobalEventsStats.Add(CurrentEvent);
+				EventToObjectsMap.FindOrAdd(CurrentEvent.Name).Add(LogItem->Name.ToString());
 			}
 
-			if (bPrintNextLine)
+			Index = GlobalEventsStats.Find(FVisualLogEntry::FLogEvent(CurrentEvent));
+			for (auto& CurrentEventTag : CurrentEvent.EventTags)
 			{
-				OutString.Append(TEXT("\n"));
+				if (!bJustAdded)
+				{
+					GlobalEventsStats[Index].EventTags.FindOrAdd(CurrentEventTag.Key) += CurrentEventTag.Value;
+				}
+
+				if (EventToObjectsMap.Contains(CurrentEvent.Name + CurrentEventTag.Key.ToString()))
+				{
+					EventToObjectsMap[CurrentEvent.Name + CurrentEventTag.Key.ToString()].AddUnique(LogItem->Name.ToString());
+				}
+				else
+				{
+					EventToObjectsMap.FindOrAdd(CurrentEvent.Name + CurrentEventTag.Key.ToString()).AddUnique(LogItem->Name.ToString());
+				}
 			}
 		}
-
-		OutString.Append(TEXT("\n"));
 	}
 
-	OutString.Append(TEXT("-===========================================-\n"));
+	OutString.Append(TEXT("-= SUMMARY =-\n"));
 
 	for (auto& CurrentEvent : GlobalEventsStats)
 	{
-		OutString.Append(FString::Printf(TEXT(" %s  (%s) occurred %d times\n"), *CurrentEvent.Name, *CurrentEvent.UserFriendlyDesc, CurrentEvent.Counter));
+		OutString.Append(FString::Printf(TEXT("%s  occurred %d times by %d owners (%s)\n"), *CurrentEvent.Name, CurrentEvent.Counter, EventToObjectsMap.Contains(CurrentEvent.Name) ? EventToObjectsMap[CurrentEvent.Name].Num() : 0, *CurrentEvent.UserFriendlyDesc));
 		for (auto& CurrentEventTag : CurrentEvent.EventTags)
 		{
-			OutString.Append(FString::Printf(TEXT("        %d  times for <%s>\n"), CurrentEventTag.Value, *CurrentEventTag.Key.ToString()));
+			const int32 ObjectsNumber = EventToObjectsMap.Contains(CurrentEvent.Name + CurrentEventTag.Key.ToString()) ? EventToObjectsMap[CurrentEvent.Name + CurrentEventTag.Key.ToString()].Num() : 0;
+			OutString.Append(FString::Printf(TEXT("%s to <%s> tag occurred %d  times by %d owners (average %.2f times each)\n"), *CurrentEvent.Name, *CurrentEventTag.Key.ToString(), CurrentEventTag.Value, ObjectsNumber, ObjectsNumber > 0 ? float(CurrentEventTag.Value) / ObjectsNumber : -1));
 		}
+		OutString.Append(TEXT("\n"));
 	}
 
 	return FText::FromString(OutString);
