Binary file format for Visual Logger (Log Visualizer can still load old file format)
Removed old Visual Logger class and files. Minor cleanups in code.

diff --git a/Engine/Config/BaseEngine.ini b/Engine/Config/BaseEngine.ini
--- a/Engine/Config/BaseEngine.ini
+++ b/Engine/Config/BaseEngine.ini
@@ -2363,6 +2363,5 @@ bTrackHistory=false
 NumberOfHistoryItemsTracked=5
 
 [VisualLogger]
-UseBinaryFileDevice=false
 FrameCacheLenght=1.0f ;in seconds, to batch log data between file serializations
 UseCompression=false ;works only with binary files
diff --git a/Engine/Source/Developer/FunctionalTesting/Private/FunctionalTest.cpp b/Engine/Source/Developer/FunctionalTesting/Private/FunctionalTest.cpp
--- a/Engine/Source/Developer/FunctionalTesting/Private/FunctionalTest.cpp
+++ b/Engine/Source/Developer/FunctionalTesting/Private/FunctionalTest.cpp
@@ -2,7 +2,7 @@
 
 #include "FunctionalTestingPrivatePCH.h"
 #include "ObjectEditorUtils.h"
-#include "VisualLog.h"
+#include "VisualLogger/VisualLogger.h"
 
 AFunctionalTest::AFunctionalTest( const FObjectInitializer& ObjectInitializer )
 	: Super(ObjectInitializer)
diff --git a/Engine/Source/Developer/GameplayDebugger/Private/GameplayDebuggingControllerComponent.cpp b/Engine/Source/Developer/GameplayDebugger/Private/GameplayDebuggingControllerComponent.cpp
--- a/Engine/Source/Developer/GameplayDebugger/Private/GameplayDebuggingControllerComponent.cpp
+++ b/Engine/Source/Developer/GameplayDebugger/Private/GameplayDebuggingControllerComponent.cpp
@@ -5,7 +5,7 @@
 #include "GameFramework/PlayerInput.h"
 #include "BehaviorTreeDelegates.h"
 #include "TimerManager.h"
-#include "VisualLog.h"
+#include "VisualLogger/VisualLogger.h"
 
 DEFINE_LOG_CATEGORY_STATIC(LogGameplayDebugging, Log, All);
 
diff --git a/Engine/Source/Developer/LogVisualizer/Private/LogVisualizer.cpp b/Engine/Source/Developer/LogVisualizer/Private/LogVisualizer.cpp
--- a/Engine/Source/Developer/LogVisualizer/Private/LogVisualizer.cpp
+++ b/Engine/Source/Developer/LogVisualizer/Private/LogVisualizer.cpp
@@ -1,12 +1,422 @@
 // Copyright 1998-2014 Epic Games, Inc. All Rights Reserved.
 
 #include "LogVisualizerPCH.h"
+#include "Misc/Base64.h"
+#include "Json.h"
 #include "CollisionDebugDrawingPublic.h"
 #include "LogVisualizerDebugActor.h"
 
 #if ENABLE_VISUAL_LOG
 //////////////////////////////////////////////////////////////////////////
 
+TSharedPtr<FJsonValue> GenerateJson(FVisualLogEntry* InEntry)
+{
+	TSharedPtr<FJsonObject> JsonEntryObject = MakeShareable(new FJsonObject);
+
+	JsonEntryObject->SetNumberField(VisualLogJson::TAG_TIMESTAMP, InEntry->TimeStamp);
+	JsonEntryObject->SetStringField(VisualLogJson::TAG_LOCATION, InEntry->Location.ToString());
+
+	const int32 StatusCategoryCount = InEntry->Status.Num();
+	const FVisualLogStatusCategory* StatusCategory = InEntry->Status.GetData();
+	TArray< TSharedPtr<FJsonValue> > JsonStatus;
+	JsonStatus.AddZeroed(StatusCategoryCount);
+
+	for (int32 StatusIdx = 0; StatusIdx < InEntry->Status.Num(); StatusIdx++, StatusCategory++)
+	{
+		TSharedPtr<FJsonObject> JsonStatusCategoryObject = MakeShareable(new FJsonObject);
+		JsonStatusCategoryObject->SetStringField(VisualLogJson::TAG_CATEGORY, StatusCategory->Category);
+
+		TArray< TSharedPtr<FJsonValue> > JsonStatusLines;
+		JsonStatusLines.AddZeroed(StatusCategory->Data.Num());
+		for (int32 LineIdx = 0; LineIdx < StatusCategory->Data.Num(); LineIdx++)
+		{
+			JsonStatusLines[LineIdx] = MakeShareable(new FJsonValueString(StatusCategory->Data[LineIdx]));
+		}
+
+		JsonStatusCategoryObject->SetArrayField(VisualLogJson::TAG_STATUSLINES, JsonStatusLines);
+		JsonStatus[StatusIdx] = MakeShareable(new FJsonValueObject(JsonStatusCategoryObject));
+	}
+	JsonEntryObject->SetArrayField(VisualLogJson::TAG_STATUS, JsonStatus);
+
+	const int32 LogLineCount = InEntry->LogLines.Num();
+	const FVisualLogLine* LogLine = InEntry->LogLines.GetData();
+	TArray< TSharedPtr<FJsonValue> > JsonLines;
+	JsonLines.AddZeroed(LogLineCount);
+
+	for (int32 LogLineIndex = 0; LogLineIndex < LogLineCount; ++LogLineIndex, ++LogLine)
+	{
+		TSharedPtr<FJsonObject> JsonLogLineObject = MakeShareable(new FJsonObject);
+		JsonLogLineObject->SetStringField(VisualLogJson::TAG_CATEGORY, LogLine->Category.ToString());
+		JsonLogLineObject->SetNumberField(VisualLogJson::TAG_VERBOSITY, LogLine->Verbosity);
+		JsonLogLineObject->SetStringField(VisualLogJson::TAG_LINE, LogLine->Line);
+		JsonLogLineObject->SetStringField(VisualLogJson::TAG_TAGNAME, LogLine->TagName.ToString());
+		JsonLogLineObject->SetNumberField(VisualLogJson::TAG_USERDATA, LogLine->UserData);
+		JsonLines[LogLineIndex] = MakeShareable(new FJsonValueObject(JsonLogLineObject));
+	}
+
+	JsonEntryObject->SetArrayField(VisualLogJson::TAG_LOGLINES, JsonLines);
+
+	const int32 ElementsToDrawCount = InEntry->ElementsToDraw.Num();
+	const FVisualLogShapeElement* Element = InEntry->ElementsToDraw.GetData();
+	TArray< TSharedPtr<FJsonValue> > JsonElementsToDraw;
+	JsonElementsToDraw.AddZeroed(ElementsToDrawCount);
+
+	for (int32 ElementIndex = 0; ElementIndex < ElementsToDrawCount; ++ElementIndex, ++Element)
+	{
+		TSharedPtr<FJsonObject> JsonElementToDrawObject = MakeShareable(new FJsonObject);
+
+		JsonElementToDrawObject->SetStringField(VisualLogJson::TAG_DESCRIPTION, Element->Description);
+		JsonElementToDrawObject->SetStringField(VisualLogJson::TAG_CATEGORY, Element->Category.ToString());
+		JsonElementToDrawObject->SetNumberField(VisualLogJson::TAG_VERBOSITY, Element->Verbosity);
+
+		const int32 EncodedTypeColorSize = uint8(Element->Type) << 24 | Element->Color << 16 | Element->Thicknes;
+		JsonElementToDrawObject->SetStringField(VisualLogJson::TAG_TYPECOLORSIZE, FString::Printf(TEXT("%d"), EncodedTypeColorSize));
+
+		const int32 PointsCount = Element->Points.Num();
+		TArray< TSharedPtr<FJsonValue> > JsonStringPoints;
+		JsonStringPoints.AddZeroed(PointsCount);
+		const FVector* PointToDraw = Element->Points.GetData();
+		for (int32 PointIndex = 0; PointIndex < PointsCount; ++PointIndex, ++PointToDraw)
+		{
+			JsonStringPoints[PointIndex] = MakeShareable(new FJsonValueString(PointToDraw->ToString()));
+		}
+		JsonElementToDrawObject->SetArrayField(VisualLogJson::TAG_POINTS, JsonStringPoints);
+
+		JsonElementsToDraw[ElementIndex] = MakeShareable(new FJsonValueObject(JsonElementToDrawObject));
+	}
+
+	JsonEntryObject->SetArrayField(VisualLogJson::TAG_ELEMENTSTODRAW, JsonElementsToDraw);
+
+	const int32 HistogramSamplesCount = InEntry->HistogramSamples.Num();
+	const FVisualLogHistogramSample* Sample = InEntry->HistogramSamples.GetData();
+	TArray<TSharedPtr<FJsonValue> > JsonHistogramSamples;
+	JsonHistogramSamples.AddZeroed(HistogramSamplesCount);
+
+	for (int32 SampleIndex = 0; SampleIndex < HistogramSamplesCount; ++SampleIndex, ++Sample)
+	{
+		TSharedPtr<FJsonObject> JsonSampleObject = MakeShareable(new FJsonObject);
+
+		JsonSampleObject->SetStringField(VisualLogJson::TAG_CATEGORY, Sample->Category.ToString());
+		JsonSampleObject->SetNumberField(VisualLogJson::TAG_VERBOSITY, Sample->Verbosity);
+		JsonSampleObject->SetStringField(VisualLogJson::TAG_HISTOGRAMSAMPLE, Sample->SampleValue.ToString());
+		JsonSampleObject->SetStringField(VisualLogJson::TAG_HISTOGRAMGRAPHNAME, Sample->GraphName.ToString());
+		JsonSampleObject->SetStringField(VisualLogJson::TAG_HISTOGRAMDATANAME, Sample->DataName.ToString());
+
+		JsonHistogramSamples[SampleIndex] = MakeShareable(new FJsonValueObject(JsonSampleObject));
+	}
+
+	JsonEntryObject->SetArrayField(VisualLogJson::TAG_HISTOGRAMSAMPLES, JsonHistogramSamples);
+
+	// generate json from events
+	{
+		const int32 EventsCount = InEntry->Events.Num();
+		TArray<TSharedPtr<FJsonValue> > JsonEventSamples;
+		JsonEventSamples.Reserve(HistogramSamplesCount);
+
+		const FVisualLogHistogramSample* Sample = InEntry->HistogramSamples.GetData();
+		for (auto& CurrentEvent : InEntry->Events)
+		{
+			TSharedPtr<FJsonObject> JsonSampleObject = MakeShareable(new FJsonObject);
+
+			JsonSampleObject->SetStringField(VisualLogJson::TAG_CATEGORY, CurrentEvent.Name);
+			JsonSampleObject->SetNumberField(VisualLogJson::TAG_VERBOSITY, CurrentEvent.Verbosity);
+			JsonSampleObject->SetStringField(VisualLogJson::TAG_TAGNAME, CurrentEvent.TagName.ToString());
+			JsonSampleObject->SetNumberField(VisualLogJson::TAG_USERDATA, CurrentEvent.UserData);
+			JsonSampleObject->SetStringField(VisualLogJson::TAG_DESCRIPTION, CurrentEvent.UserFriendlyDesc);
+			JsonSampleObject->SetNumberField(VisualLogJson::TAG_COUNTER, CurrentEvent.Counter);
+
+			TArray<TSharedPtr<FJsonValue> > JsonEventTags;
+			JsonEventTags.Reserve(HistogramSamplesCount);
+			for (auto& CurrentTag : CurrentEvent.EventTags)
+			{
+				TSharedPtr<FJsonObject> JsonEventObject = MakeShareable(new FJsonObject);
+				JsonEventObject->SetStringField(VisualLogJson::TAG_EVENTNAME, CurrentTag.Key.ToString());
+				JsonEventObject->SetNumberField(VisualLogJson::TAG_COUNTER, CurrentTag.Value);
+
+				JsonEventTags.Add(MakeShareable(new FJsonValueObject(JsonEventObject)));
+			}
+			JsonSampleObject->SetArrayField(VisualLogJson::TAG_EVENTAGS, JsonEventTags);
+
+			JsonEventSamples.Add(MakeShareable(new FJsonValueObject(JsonSampleObject)));
+		}
+
+		JsonEntryObject->SetArrayField(VisualLogJson::TAG_EVENTSAMPLES, JsonEventSamples);
+	}
+
+	int32 DataBlockIndex = 0;
+	TArray<TSharedPtr<FJsonValue> > JasonDataBlocks;
+	JasonDataBlocks.AddZeroed(InEntry->DataBlocks.Num());
+	for (const auto CurrentData : InEntry->DataBlocks)
+	{
+		TSharedPtr<FJsonObject> JsonSampleObject = MakeShareable(new FJsonObject);
+
+		TArray<uint8> CompressedData;
+		const int32 UncompressedSize = CurrentData.Data.Num();
+		const int32 HeaderSize = sizeof(int32);
+		CompressedData.Init(0, HeaderSize + FMath::TruncToInt(1.1f * UncompressedSize));
+
+		int32 CompressedSize = CompressedData.Num() - HeaderSize;
+		uint8* DestBuffer = CompressedData.GetData();
+		FMemory::Memcpy(DestBuffer, &UncompressedSize, HeaderSize);
+		DestBuffer += HeaderSize;
+
+		FCompression::CompressMemory((ECompressionFlags)(COMPRESS_ZLIB | COMPRESS_BiasMemory), (void*)DestBuffer, CompressedSize, (void*)CurrentData.Data.GetData(), UncompressedSize);
+		CompressedData.SetNum(CompressedSize + HeaderSize, true);
+		const FString CurrentDataAsString = FBase64::Encode(CompressedData);
+
+		JsonSampleObject->SetStringField(VisualLogJson::TAG_CATEGORY, CurrentData.Category.ToString());
+		JsonSampleObject->SetStringField(VisualLogJson::TAG_TAGNAME, CurrentData.TagName.ToString());
+		JsonSampleObject->SetStringField(VisualLogJson::TAG_DATABLOCK_DATA, CurrentDataAsString);
+		JsonSampleObject->SetNumberField(VisualLogJson::TAG_VERBOSITY, CurrentData.Verbosity);
+
+		JasonDataBlocks[DataBlockIndex++] = MakeShareable(new FJsonValueObject(JsonSampleObject));
+	}
+	JsonEntryObject->SetArrayField(VisualLogJson::TAG_DATABLOCK, JasonDataBlocks);
+
+	return MakeShareable(new FJsonValueObject(JsonEntryObject));
+}
+
+TSharedPtr<FVisualLogEntry> JsonToVisualLogEntry(TSharedPtr<FJsonValue> FromJson)
+{
+	TSharedPtr<FJsonObject> JsonEntryObject = FromJson->AsObject();
+	if (JsonEntryObject.IsValid() == false)
+	{
+		return NULL;
+	}
+
+	TSharedPtr<FVisualLogEntry> VisLogEntry = MakeShareable(new FVisualLogEntry());
+	VisLogEntry->TimeStamp = JsonEntryObject->GetNumberField(VisualLogJson::TAG_TIMESTAMP);
+	VisLogEntry->Location.InitFromString(JsonEntryObject->GetStringField(VisualLogJson::TAG_LOCATION));
+
+	TArray< TSharedPtr<FJsonValue> > JsonStatus = JsonEntryObject->GetArrayField(VisualLogJson::TAG_STATUS);
+	if (JsonStatus.Num() > 0)
+	{
+		VisLogEntry->Status.AddZeroed(JsonStatus.Num());
+		for (int32 StatusIndex = 0; StatusIndex < JsonStatus.Num(); ++StatusIndex)
+		{
+			TSharedPtr<FJsonObject> JsonStatusCategory = JsonStatus[StatusIndex]->AsObject();
+			if (JsonStatusCategory.IsValid())
+			{
+				VisLogEntry->Status[StatusIndex].Category = JsonStatusCategory->GetStringField(VisualLogJson::TAG_CATEGORY);
+
+				TArray< TSharedPtr<FJsonValue> > JsonStatusLines = JsonStatusCategory->GetArrayField(VisualLogJson::TAG_STATUSLINES);
+				if (JsonStatusLines.Num() > 0)
+				{
+					VisLogEntry->Status[StatusIndex].Data.AddZeroed(JsonStatusLines.Num());
+					for (int32 LineIdx = 0; LineIdx < JsonStatusLines.Num(); LineIdx++)
+					{
+						VisLogEntry->Status[StatusIndex].Data.Add(JsonStatusLines[LineIdx]->AsString());
+					}
+				}
+			}
+		}
+	}
+
+	TArray< TSharedPtr<FJsonValue> > JsonLines = JsonEntryObject->GetArrayField(VisualLogJson::TAG_LOGLINES);
+	if (JsonLines.Num() > 0)
+	{
+		VisLogEntry->LogLines.AddZeroed(JsonLines.Num());
+		for (int32 LogLineIndex = 0; LogLineIndex < JsonLines.Num(); ++LogLineIndex)
+		{
+			TSharedPtr<FJsonObject> JsonLogLine = JsonLines[LogLineIndex]->AsObject();
+			if (JsonLogLine.IsValid())
+			{
+				VisLogEntry->LogLines[LogLineIndex].Category = FName(*(JsonLogLine->GetStringField(VisualLogJson::TAG_CATEGORY)));
+				VisLogEntry->LogLines[LogLineIndex].Verbosity = TEnumAsByte<ELogVerbosity::Type>((uint8)FMath::TruncToInt(JsonLogLine->GetNumberField(VisualLogJson::TAG_VERBOSITY)));
+				VisLogEntry->LogLines[LogLineIndex].Line = JsonLogLine->GetStringField(VisualLogJson::TAG_LINE);
+				VisLogEntry->LogLines[LogLineIndex].TagName = FName(*(JsonLogLine->GetStringField(VisualLogJson::TAG_TAGNAME)));
+				VisLogEntry->LogLines[LogLineIndex].UserData = (int64)(JsonLogLine->GetNumberField(VisualLogJson::TAG_USERDATA));
+
+			}
+		}
+	}
+
+	TArray< TSharedPtr<FJsonValue> > JsonElementsToDraw = JsonEntryObject->GetArrayField(VisualLogJson::TAG_ELEMENTSTODRAW);
+	if (JsonElementsToDraw.Num() > 0)
+	{
+		VisLogEntry->ElementsToDraw.Reserve(JsonElementsToDraw.Num());
+		for (int32 ElementIndex = 0; ElementIndex < JsonElementsToDraw.Num(); ++ElementIndex)
+		{
+			TSharedPtr<FJsonObject> JsonElementObject = JsonElementsToDraw[ElementIndex]->AsObject();
+			if (JsonElementObject.IsValid())
+			{
+				FVisualLogShapeElement& Element = VisLogEntry->ElementsToDraw[VisLogEntry->ElementsToDraw.Add(FVisualLogShapeElement())];
+
+				Element.Description = JsonElementObject->GetStringField(VisualLogJson::TAG_DESCRIPTION);
+				Element.Category = FName(*(JsonElementObject->GetStringField(VisualLogJson::TAG_CATEGORY)));
+				Element.Verbosity = TEnumAsByte<ELogVerbosity::Type>((uint8)FMath::TruncToInt(JsonElementObject->GetNumberField(VisualLogJson::TAG_VERBOSITY)));
+
+				// Element->Type << 24 | Element->Color << 16 | Element->Thicknes;
+				int32 EncodedTypeColorSize = 0;
+				TTypeFromString<int32>::FromString(EncodedTypeColorSize, *(JsonElementObject->GetStringField(VisualLogJson::TAG_TYPECOLORSIZE)));
+				Element.Type = EVisualLoggerShapeElement(EncodedTypeColorSize >> 24);
+				Element.Color = (EncodedTypeColorSize << 8) >> 24;
+				Element.Thicknes = (EncodedTypeColorSize << 16) >> 16;
+
+				TArray< TSharedPtr<FJsonValue> > JsonPoints = JsonElementObject->GetArrayField(VisualLogJson::TAG_POINTS);
+				if (JsonPoints.Num() > 0)
+				{
+					Element.Points.AddUninitialized(JsonPoints.Num());
+					for (int32 PointIndex = 0; PointIndex < JsonPoints.Num(); ++PointIndex)
+					{
+						Element.Points[PointIndex].InitFromString(JsonPoints[PointIndex]->AsString());
+					}
+				}
+			}
+		}
+	}
+
+	TArray< TSharedPtr<FJsonValue> > JsonHistogramSamples = JsonEntryObject->GetArrayField(VisualLogJson::TAG_HISTOGRAMSAMPLES);
+	if (JsonHistogramSamples.Num() > 0)
+	{
+		VisLogEntry->HistogramSamples.Reserve(JsonHistogramSamples.Num());
+		for (int32 SampleIndex = 0; SampleIndex < JsonHistogramSamples.Num(); ++SampleIndex)
+		{
+			TSharedPtr<FJsonObject> JsonSampleObject = JsonHistogramSamples[SampleIndex]->AsObject();
+			if (JsonSampleObject.IsValid())
+			{
+				FVisualLogHistogramSample& Sample = VisLogEntry->HistogramSamples[VisLogEntry->HistogramSamples.Add(FVisualLogHistogramSample())];
+
+				Sample.Category = FName(*(JsonSampleObject->GetStringField(VisualLogJson::TAG_CATEGORY)));
+				Sample.Verbosity = TEnumAsByte<ELogVerbosity::Type>((uint8)FMath::TruncToInt(JsonSampleObject->GetNumberField(VisualLogJson::TAG_VERBOSITY)));
+				Sample.GraphName = FName(*(JsonSampleObject->GetStringField(VisualLogJson::TAG_HISTOGRAMGRAPHNAME)));
+				Sample.DataName = FName(*(JsonSampleObject->GetStringField(VisualLogJson::TAG_HISTOGRAMDATANAME)));
+				Sample.SampleValue.InitFromString(JsonSampleObject->GetStringField(VisualLogJson::TAG_HISTOGRAMSAMPLE));
+			}
+		}
+	}
+
+	TArray< TSharedPtr<FJsonValue> > JsonEventSamples = JsonEntryObject->GetArrayField(VisualLogJson::TAG_EVENTSAMPLES);
+	if (JsonEventSamples.Num() > 0)
+	{
+		VisLogEntry->Events.Reserve(JsonEventSamples.Num());
+		for (int32 SampleIndex = 0; SampleIndex < JsonEventSamples.Num(); ++SampleIndex)
+		{
+			TSharedPtr<FJsonObject> JsonSampleObject = JsonEventSamples[SampleIndex]->AsObject();
+			if (JsonSampleObject.IsValid())
+			{
+				FVisualLogEvent& Event = VisLogEntry->Events[VisLogEntry->Events.Add(FVisualLogEvent())];
+				Event.Name = JsonSampleObject->GetStringField(VisualLogJson::TAG_CATEGORY);
+				Event.Verbosity = TEnumAsByte<ELogVerbosity::Type>((uint8)FMath::TruncToInt(JsonSampleObject->GetNumberField(VisualLogJson::TAG_VERBOSITY)));
+				Event.TagName = FName(*(JsonSampleObject->GetStringField(VisualLogJson::TAG_TAGNAME)));
+				Event.UserData = (int64)(JsonSampleObject->GetNumberField(VisualLogJson::TAG_USERDATA));
+				Event.UserFriendlyDesc = JsonSampleObject->GetStringField(VisualLogJson::TAG_DESCRIPTION);
+				Event.Counter = (int64)(JsonSampleObject->GetNumberField(VisualLogJson::TAG_COUNTER));
+
+				TArray< TSharedPtr<FJsonValue> > JsonEventTags = JsonSampleObject->GetArrayField(VisualLogJson::TAG_EVENTAGS);
+				if (JsonEventTags.Num() > 0)
+				{
+					Event.EventTags.Reserve(JsonEventSamples.Num());
+					for (int32 TagIndex = 0; TagIndex < JsonEventTags.Num(); ++TagIndex)
+					{
+						TSharedPtr<FJsonObject> JsonTagObject = JsonEventTags[TagIndex]->AsObject();
+						if (JsonTagObject.IsValid())
+						{
+							FString TagName = JsonTagObject->GetStringField(VisualLogJson::TAG_EVENTNAME);
+							int32 TagCounter = JsonTagObject->GetNumberField(VisualLogJson::TAG_COUNTER);
+							Event.EventTags.Add(*TagName, TagCounter);
+						}
+					}
+				}
+			}
+		}
+	}
+
+	TArray< TSharedPtr<FJsonValue> > JasonDataBlocks = JsonEntryObject->GetArrayField(VisualLogJson::TAG_DATABLOCK);
+	if (JasonDataBlocks.Num() > 0)
+	{
+		VisLogEntry->DataBlocks.Reserve(JasonDataBlocks.Num());
+		for (int32 SampleIndex = 0; SampleIndex < JasonDataBlocks.Num(); ++SampleIndex)
+		{
+			TSharedPtr<FJsonObject> JsonSampleObject = JasonDataBlocks[SampleIndex]->AsObject();
+			if (JsonSampleObject.IsValid())
+			{
+				FVisualLogDataBlock& Sample = VisLogEntry->DataBlocks[VisLogEntry->DataBlocks.Add(FVisualLogDataBlock())];
+				Sample.TagName = FName(*(JsonSampleObject->GetStringField(VisualLogJson::TAG_TAGNAME)));
+				Sample.Category = FName(*(JsonSampleObject->GetStringField(VisualLogJson::TAG_CATEGORY)));
+				Sample.Verbosity = TEnumAsByte<ELogVerbosity::Type>((uint8)FMath::TruncToInt(JsonSampleObject->GetNumberField(VisualLogJson::TAG_VERBOSITY)));
+
+				// decode data from Base64 string
+				const FString DataBlockAsCompressedString = JsonSampleObject->GetStringField(VisualLogJson::TAG_DATABLOCK_DATA);
+				TArray<uint8> CompressedDataBlock;
+				FBase64::Decode(DataBlockAsCompressedString, CompressedDataBlock);
+
+				// uncompress decoded data to get final TArray<uint8>
+				int32 UncompressedSize = 0;
+				const int32 HeaderSize = sizeof(int32);
+				uint8* SrcBuffer = (uint8*)CompressedDataBlock.GetData();
+				FMemory::Memcpy(&UncompressedSize, SrcBuffer, HeaderSize);
+				SrcBuffer += HeaderSize;
+				const int32 CompressedSize = CompressedDataBlock.Num() - HeaderSize;
+
+				Sample.Data.AddZeroed(UncompressedSize);
+				FCompression::UncompressMemory((ECompressionFlags)(COMPRESS_ZLIB), (void*)Sample.Data.GetData(), UncompressedSize, SrcBuffer, CompressedSize);
+			}
+		}
+	}
+
+	return VisLogEntry;
+}
+//----------------------------------------------------------------------//
+// FActorsVisLog
+//----------------------------------------------------------------------//
+FActorsVisLog::FActorsVisLog(FName InName, TArray<TWeakObjectPtr<UObject> >* Children)
+	: Name(InName)
+{
+	Entries.Reserve(VisLogInitialSize);
+}
+
+FActorsVisLog::FActorsVisLog(const class UObject* Object, TArray<TWeakObjectPtr<UObject> >* Children)
+	: Name(Object->GetFName())
+	, FullName(Object->GetFullName())
+{
+	Entries.Reserve(VisLogInitialSize);
+}
+
+FActorsVisLog::FActorsVisLog(TSharedPtr<FJsonValue> FromJson)
+{
+	TSharedPtr<FJsonObject> JsonLogObject = FromJson->AsObject();
+	if (JsonLogObject.IsValid() == false)
+	{
+		return;
+	}
+
+	Name = FName(*JsonLogObject->GetStringField(VisualLogJson::TAG_NAME));
+	FullName = JsonLogObject->GetStringField(VisualLogJson::TAG_FULLNAME);
+
+	TArray< TSharedPtr<FJsonValue> > JsonEntries = JsonLogObject->GetArrayField(VisualLogJson::TAG_ENTRIES);
+	if (JsonEntries.Num() > 0)
+	{
+		Entries.AddZeroed(JsonEntries.Num());
+		for (int32 EntryIndex = 0; EntryIndex < JsonEntries.Num(); ++EntryIndex)
+		{
+			Entries[EntryIndex] = JsonToVisualLogEntry(JsonEntries[EntryIndex]);
+		}
+	}
+}
+
+TSharedPtr<FJsonValue> FActorsVisLog::ToJson() const
+{
+	TSharedPtr<FJsonObject> JsonLogObject = MakeShareable(new FJsonObject);
+
+	JsonLogObject->SetStringField(VisualLogJson::TAG_NAME, Name.ToString());
+	JsonLogObject->SetStringField(VisualLogJson::TAG_FULLNAME, FullName);
+
+	const TSharedPtr<FVisualLogEntry>* Entry = Entries.GetData();
+	const int32 EntryCount = Entries.Num();
+	TArray< TSharedPtr<FJsonValue> > JsonLogEntries;
+	JsonLogEntries.AddZeroed(EntryCount);
+
+	for (int32 EntryIndex = 0; EntryIndex < EntryCount; ++EntryIndex, ++Entry)
+	{
+		JsonLogEntries[EntryIndex] = GenerateJson((*Entry).Get());
+	}
+
+	JsonLogObject->SetArrayField(VisualLogJson::TAG_ENTRIES, JsonLogEntries);
+
+	return MakeShareable(new FJsonValueObject(JsonLogObject));
+}
+
 void FLogVisualizer::SummonUI(UWorld* InWorld) 
 {
 	UE_LOG(LogLogVisualizer, Log, TEXT("Opening LogVisualizer..."));
@@ -19,9 +429,21 @@ void FLogVisualizer::SummonUI(UWorld* InWorld)
 		}
 
 		World = InWorld;
-		FVisualLog& VisualLog = FVisualLog::GetStatic();
-		VisualLog.RegisterNewLogsObserver(FVisualLog::FOnNewLogCreatedDelegate::CreateRaw(this, &FLogVisualizer::OnNewLog));
-		PullDataFromVisualLog(VisualLog);
+
+		auto &Devices = FVisualLogger::Get().GetDevices();
+		for (auto CurrentDevice : Devices)
+		{
+			if (CurrentDevice != NULL && CurrentDevice->HasFlags(EVisualLoggerDeviceFlags::StoreLogsLocally))
+			{
+				TArray<FVisualLogDevice::FVisualLogEntryItem> RecordedLogs;
+				CurrentDevice->GetRecordedLogs(RecordedLogs);
+				PullDataFromVisualLog(RecordedLogs);
+				break;
+			}
+
+		}
+		FVisualLogger& VisualLog = FVisualLogger::Get();
+		VisualLog.AddDevice(this);
 
 		// Give window to slate
 		if (!LogWindow.IsValid())
@@ -77,7 +499,8 @@ bool FLogVisualizer::IsOpenUI(UWorld* InWorld)
 
 void FLogVisualizer::CleanUp()
 {
-	FVisualLog::GetStatic().ClearNewLogsObserver();
+	FVisualLogger& VisualLog = FVisualLogger::Get();
+	VisualLog.RemoveDevice(this);
 }
 
 class AActor* FLogVisualizer::GetHelperActor(class UWorld* InWorld)
@@ -104,15 +527,27 @@ class AActor* FLogVisualizer::GetHelperActor(class UWorld* InWorld)
 	return DebugActor.Get();
 }
 
-void FLogVisualizer::PullDataFromVisualLog(const FVisualLog& VisualLog)
+void FLogVisualizer::PullDataFromVisualLog(const TArray<FVisualLogDevice::FVisualLogEntryItem>& RecordedLogs)
 {
 	Logs.Reset();
-	const FVisualLog::FLogsMap* LogsMap = VisualLog.GetLogs();
-	for (FVisualLog::FLogsMap::TConstIterator MapIt(*LogsMap); MapIt; ++MapIt)
+	HelperMap.Reset();
+
+	for (auto& CurrentLog : RecordedLogs)
 	{
-		Logs.Add(MapIt.Value());
+		if (HelperMap.Contains(CurrentLog.OwnerName) == false)
+		{
+			HelperMap.Add(CurrentLog.OwnerName, MakeShareable(new FActorsVisLog(CurrentLog.OwnerName, NULL)));
+		}
+		HelperMap[CurrentLog.OwnerName]->Entries.Add(MakeShareable(new FVisualLogEntry(CurrentLog.Entry)));
+	}
+
+	for (auto& CurrentLog : HelperMap)
+	{
+		Logs.Add(CurrentLog.Value);
 		LogAddedEvent.Broadcast();
 	}
+
+
 }
 
 void FLogVisualizer::OnNewLog(const UObject* Object, TSharedPtr<FActorsVisLog> Log)
@@ -121,6 +556,22 @@ void FLogVisualizer::OnNewLog(const UObject* Object, TSharedPtr<FActorsVisLog> L
 	LogAddedEvent.Broadcast();
 }
 
+void FLogVisualizer::Serialize(const class UObject* LogOwner, FName OwnerName, const FVisualLogEntry& CurrentLog)
+{
+	const bool bContainsLogOwner = HelperMap.Contains(OwnerName);
+	if (!bContainsLogOwner)
+	{
+		HelperMap.Add(OwnerName, MakeShareable(new FActorsVisLog(OwnerName, NULL)));
+	}
+	HelperMap[OwnerName]->Entries.Add(MakeShareable(new FVisualLogEntry(CurrentLog)));
+
+	if (!bContainsLogOwner)
+	{
+		Logs.Add(HelperMap[OwnerName]);
+		LogAddedEvent.Broadcast();
+	}
+}
+
 void FLogVisualizer::AddLoadedLog(TSharedPtr<FActorsVisLog> Log)
 {
 	for (int32 Index = 0; Index < Logs.Num(); ++Index)
diff --git a/Engine/Source/Developer/LogVisualizer/Private/LogVisualizer.h b/Engine/Source/Developer/LogVisualizer/Private/LogVisualizer.h
--- a/Engine/Source/Developer/LogVisualizer/Private/LogVisualizer.h
+++ b/Engine/Source/Developer/LogVisualizer/Private/LogVisualizer.h
@@ -2,21 +2,73 @@
 
 #pragma once
 
-#include "VisualLog.h"
+#include "VisualLogger/VisualLogger.h"
 
 #if ENABLE_VISUAL_LOG
 
-class FVisualLog;
-struct FActorsVisLog;
+namespace VisualLogJson
+{
+	static const FString TAG_NAME = TEXT("Name");
+	static const FString TAG_FULLNAME = TEXT("FullName");
+	static const FString TAG_ENTRIES = TEXT("Entries");
+	static const FString TAG_TIMESTAMP = TEXT("TimeStamp");
+	static const FString TAG_LOCATION = TEXT("Location");
+	static const FString TAG_STATUS = TEXT("Status");
+	static const FString TAG_STATUSLINES = TEXT("StatusLines");
+	static const FString TAG_CATEGORY = TEXT("Category");
+	static const FString TAG_LINE = TEXT("Line");
+	static const FString TAG_VERBOSITY = TEXT("Verb");
+	static const FString TAG_LOGLINES = TEXT("LogLines");
+	static const FString TAG_DESCRIPTION = TEXT("Description");
+	static const FString TAG_TYPECOLORSIZE = TEXT("TypeColorSize");
+	static const FString TAG_POINTS = TEXT("Points");
+	static const FString TAG_ELEMENTSTODRAW = TEXT("ElementsToDraw");
+	static const FString TAG_TAGNAME = TEXT("DataBlockTagName");
+	static const FString TAG_USERDATA = TEXT("UserData");
+	static const FString TAG_CAPSULESTODRAW = TEXT("CapsulesToDraw");
+	static const FString TAG_COUNTER = TEXT("Counter");
+	static const FString TAG_EVENTNAME = TEXT("EventName");
+	static const FString TAG_EVENTAGS = TEXT("EventTags");
+	static const FString TAG_EVENTSAMPLES = TEXT("EventSamples");
+
+	static const FString TAG_HISTOGRAMSAMPLES = TEXT("HistogramSamples");
+	static const FString TAG_HISTOGRAMSAMPLE = TEXT("Sample");
+	static const FString TAG_HISTOGRAMGRAPHNAME = TEXT("GraphName");
+	static const FString TAG_HISTOGRAMDATANAME = TEXT("DataName");
+
+	static const FString TAG_DATABLOCK = TEXT("DataBlock");
+	static const FString TAG_DATABLOCK_DATA = TEXT("DataBlockData");
+
+	static const FString TAG_LOGS = TEXT("Logs");
+}
+
+struct FActorsVisLog
+{
+	static const int VisLogInitialSize = 8;
+
+	FName Name;
+	FString FullName;
+	TArray<TSharedPtr<FVisualLogEntry> > Entries;
+
+	FActorsVisLog(FName Name, TArray<TWeakObjectPtr<UObject> >* Children);
+	FActorsVisLog(const class UObject* Object, TArray<TWeakObjectPtr<UObject> >* Children);
+	FActorsVisLog(TSharedPtr<FJsonValue> FromJson);
+
+	TSharedPtr<FJsonValue> ToJson() const;
+};
+
 
 /** Actual implementation of LogVisualizer, private inside module */
-class FLogVisualizer : public ILogVisualizer
+class FLogVisualizer : public ILogVisualizer, public FVisualLogDevice
 {
 public:
 	virtual void SummonUI(UWorld* InWorld) override;
 	virtual void CloseUI(UWorld* InWorld) override;
 	virtual bool IsOpenUI(UWorld* InWorld) override;
 
+	virtual void Serialize(const class UObject* LogOwner, FName OwnerName, const FVisualLogEntry& LogEntry);
+	virtual bool HasFlags(int32 InFlags) const { return !!(InFlags & EVisualLoggerDeviceFlags::NoFlags); }
+
 	virtual bool IsRecording();
 	// End ILogVisualizer interface
 
@@ -37,7 +89,7 @@ public:
 
 	void CleanUp();
 
-	void PullDataFromVisualLog(const FVisualLog& VisualLog);
+	void PullDataFromVisualLog(const TArray<FVisualLogDevice::FVisualLogEntryItem>& RecordedLogs);
 
 	int32 GetLogIndexForActor(const class AActor* Actor);
 
@@ -60,6 +112,8 @@ private:
 	TWeakObjectPtr<UWorld>	World;
 	TWeakPtr<SWindow>		LogWindow;
 	TWeakObjectPtr<class ALogVisualizerDebugActor> DebugActor;
+	TMap<FName, TSharedPtr<FActorsVisLog>> HelperMap;
+
 };
 
 #endif //ENABLE_VISUAL_LOG
\ No newline at end of file
diff --git a/Engine/Source/Developer/LogVisualizer/Private/SLogBar.h b/Engine/Source/Developer/LogVisualizer/Private/SLogBar.h
--- a/Engine/Source/Developer/LogVisualizer/Private/SLogBar.h
+++ b/Engine/Source/Developer/LogVisualizer/Private/SLogBar.h
@@ -2,7 +2,7 @@
 
 #pragma once
 
-#include "VisualLog.h"
+#include "VisualLogger/VisualLogger.h"
 
 struct FLogEntryBarProxy
 {
diff --git a/Engine/Source/Developer/LogVisualizer/Private/SLogVisualizer.cpp b/Engine/Source/Developer/LogVisualizer/Private/SLogVisualizer.cpp
--- a/Engine/Source/Developer/LogVisualizer/Private/SLogVisualizer.cpp
+++ b/Engine/Source/Developer/LogVisualizer/Private/SLogVisualizer.cpp
@@ -19,6 +19,7 @@
 #	include "Editor/UnrealEd/Public/EditorViewportClient.h"
 #endif
 
+#include "VIsualLogger/VisualLoggerBinaryFileDevice.h"
 #include "GameplayDebuggingComponent.h"
 #include "LogVisualizerModule.h"
 #include "SLogVisualizerReport.h"
@@ -36,10 +37,9 @@ const FName SLogVisualizer::NAME_LogTimeSpan = TEXT("LogTimeSpan");
 
 namespace LogVisualizer
 {
-	static const FString LogFileExtensionPure = TEXT("vlog");
 	static const FString LogFileDescription = LOCTEXT("FileTypeDescription", "Visual Log File").ToString();
-	static const FString LogFileExtension = FString::Printf(TEXT("*.%s"), *LogFileExtensionPure);
-	static const FString FileTypes = FString::Printf( TEXT("%s (%s)|%s"), *LogFileDescription, *LogFileExtension, *LogFileExtension );
+	static const FString LoadFileTypes = FString::Printf(TEXT("%s (*.vlog;*.%s)|*.vlog;*.%s"), *LogFileDescription, VISLOG_FILENAME_EXT, VISLOG_FILENAME_EXT);
+	static const FString SaveFileTypes = FString::Printf(TEXT("%s (*.%s)|*.%s"), *LogFileDescription, VISLOG_FILENAME_EXT, VISLOG_FILENAME_EXT);
 }
 
 FColor SLogVisualizer::ColorPalette[] = {
@@ -574,6 +574,7 @@ SLogVisualizer::~SLogVisualizer()
 	UGameplayDebuggingComponent::OnDebuggingTargetChangedDelegate.RemoveAll(this);
 	LogVisualizer->OnLogAdded().RemoveAll(this);
 	UDebugDrawService::Unregister(DrawingOnCanvasDelegate);
+
 	InvalidateCanvas();
 }
 
@@ -770,7 +771,7 @@ void SLogVisualizer::UpdateVisibleEntriesCache(const TSharedPtr<FActorsVisLog>&
 			{
 				for (int32 LogLineIndex = 0; LogLineIndex < CurrentEntry->LogLines.Num(); ++LogLineIndex)
 				{
-					const FVisualLogEntry::FLogLine &CurrentLine = CurrentEntry->LogLines[LogLineIndex];
+					const FVisualLogLine &CurrentLine = CurrentEntry->LogLines[LogLineIndex];
 					const FString CurrentCategory = CurrentLine.Category.ToString();
 					if (FilterListPtr->IsFilterEnabled(CurrentCategory, CurrentLine.Verbosity))
 					{
@@ -789,7 +790,7 @@ void SLogVisualizer::UpdateVisibleEntriesCache(const TSharedPtr<FActorsVisLog>&
 			for (int32 ElementIndex = 0; ElementIndex < CurrentEntry->ElementsToDraw.Num(); ++ElementIndex)
 			{
 				const FString CurrentCategory = CurrentEntry->ElementsToDraw[ElementIndex].Category.ToString();
-				FVisualLogEntry::FElementToDraw &CurrentElement = CurrentEntry->ElementsToDraw[ElementIndex];
+				FVisualLogShapeElement &CurrentElement = CurrentEntry->ElementsToDraw[ElementIndex];
 				if (CurrentElement.Category != NAME_None && (FilterListPtr->IsFilterEnabled(CurrentCategory, CurrentElement.Verbosity)))
 				{
 					OutEntriesCached[CachedLogIndex].CachedEntries.AddUnique(CurrentEntry);
@@ -804,7 +805,7 @@ void SLogVisualizer::UpdateVisibleEntriesCache(const TSharedPtr<FActorsVisLog>&
 
 			for (int32 SampleIndex = 0; SampleIndex < CurrentEntry->HistogramSamples.Num(); ++SampleIndex)
 			{
-				const FVisualLogEntry::FHistogramSample &CurrentSample = CurrentEntry->HistogramSamples[SampleIndex];
+				const FVisualLogHistogramSample &CurrentSample = CurrentEntry->HistogramSamples[SampleIndex];
 				const FName CurrentCategory = CurrentSample.Category;
 				const FName CurrentDataName = CurrentSample.DataName;
 
@@ -828,7 +829,7 @@ void SLogVisualizer::UpdateVisibleEntriesCache(const TSharedPtr<FActorsVisLog>&
 			{
 				for (int32 LogLineIndex = 0; LogLineIndex < CurrentEntry->Events.Num(); ++LogLineIndex)
 				{
-					const FVisualLogEntry::FLogEvent &CurrentLine = CurrentEntry->Events[LogLineIndex];
+					const FVisualLogEvent &CurrentLine = CurrentEntry->Events[LogLineIndex];
 					const FString CurrentCategory = CurrentLine.Name;
 					if (FilterListPtr->IsFilterEnabled(CurrentCategory, CurrentLine.Verbosity))
 					{
@@ -904,7 +905,7 @@ void SLogVisualizer::GetVisibleEntries(const TSharedPtr<FActorsVisLog>& Log, TAr
 				{
 					for (int32 j = 0; j < CurrentEntry->LogLines.Num(); ++j)
 					{
-						const FVisualLogEntry::FLogLine &CurrentLine = CurrentEntry->LogLines[j];
+						const FVisualLogLine &CurrentLine = CurrentEntry->LogLines[j];
 						const FString CurrentCategory = CurrentLine.Category.ToString();
 						if (FilterListPtr->IsFilterEnabled(CurrentCategory, CurrentLine.Verbosity))
 						{
@@ -923,7 +924,7 @@ void SLogVisualizer::GetVisibleEntries(const TSharedPtr<FActorsVisLog>& Log, TAr
 				for (int32 j = 0; j < CurrentEntry->ElementsToDraw.Num(); ++j)
 				{
 					const FString CurrentCategory = CurrentEntry->ElementsToDraw[j].Category.ToString();
-					FVisualLogEntry::FElementToDraw &CurrentElement = CurrentEntry->ElementsToDraw[j];
+					FVisualLogShapeElement &CurrentElement = CurrentEntry->ElementsToDraw[j];
 					if (CurrentElement.Category != NAME_None && (FilterListPtr->IsFilterEnabled(CurrentCategory, CurrentElement.Verbosity)))
 					{
 						OutEntriesCached[CachedLogIndex].CachedEntries.AddUnique(CurrentEntry);
@@ -938,7 +939,7 @@ void SLogVisualizer::GetVisibleEntries(const TSharedPtr<FActorsVisLog>& Log, TAr
 
 				for (int32 SampleIndex = 0; SampleIndex < CurrentEntry->HistogramSamples.Num(); ++SampleIndex)
 				{
-					const FVisualLogEntry::FHistogramSample &CurrentSample = CurrentEntry->HistogramSamples[SampleIndex];
+					const FVisualLogHistogramSample &CurrentSample = CurrentEntry->HistogramSamples[SampleIndex];
 					const FName CurrentCategory = CurrentSample.Category;
 					const FName CurrentDataName = CurrentSample.DataName;
 
@@ -960,7 +961,7 @@ void SLogVisualizer::GetVisibleEntries(const TSharedPtr<FActorsVisLog>& Log, TAr
 				for (int32 j = 0; j < CurrentEntry->Events.Num(); ++j)
 				{
 					const FString CurrentCategory = CurrentEntry->Events[j].Name;
-					FVisualLogEntry::FLogEvent &CurrentElement = CurrentEntry->Events[j];
+					FVisualLogEvent &CurrentElement = CurrentEntry->Events[j];
 					if (CurrentElement.Name.Len() > 0 && (FilterListPtr->IsFilterEnabled(CurrentCategory, CurrentElement.Verbosity)))
 					{
 						OutEntriesCached[CachedLogIndex].CachedEntries.AddUnique(CurrentEntry);
@@ -1862,7 +1863,7 @@ void SLogVisualizer::ShowEntry(const FVisualLogEntry* LogEntry)
 	UpdateStatusItems(LogEntry);
 	LogEntryLines.Reset();
 	
-	const FVisualLogEntry::FLogLine* LogLine = LogEntry->LogLines.GetData();
+	const FVisualLogLine* LogLine = LogEntry->LogLines.GetData();
 	for (int LineIndex = 0; LineIndex < LogEntry->LogLines.Num(); ++LineIndex, ++LogLine)
 	{
 		bool bShowLine = true;
@@ -2246,7 +2247,7 @@ void SLogVisualizer::DrawOnCanvas(UCanvas* Canvas, APlayerController*)
 				const int32 SamplesNum = CurrentEntry->HistogramSamples.Num();
 				for (int32 SampleIndex = 0; SampleIndex < SamplesNum; ++SampleIndex)
 				{
-					FVisualLogEntry::FHistogramSample CurrentSample = CurrentEntry->HistogramSamples[SampleIndex];
+					FVisualLogHistogramSample CurrentSample = CurrentEntry->HistogramSamples[SampleIndex];
 
 					const FName CurrentCategory = CurrentSample.Category;
 					const FName CurrentGraphName = CurrentSample.GraphName;
@@ -2392,7 +2393,7 @@ void SLogVisualizer::DrawOnCanvas(UCanvas* Canvas, APlayerController*)
 				}
 			}
 
-			const FVisualLogEntry::FElementToDraw* ElementToDraw = Entry->ElementsToDraw.GetData();
+			const FVisualLogShapeElement* ElementToDraw = Entry->ElementsToDraw.GetData();
 			const int32 ElementsCount = Entry->ElementsToDraw.Num();
 			
 			for (int32 ElementIndex = 0; ElementIndex < ElementsCount; ++ElementIndex, ++ElementToDraw)
@@ -2407,7 +2408,7 @@ void SLogVisualizer::DrawOnCanvas(UCanvas* Canvas, APlayerController*)
 
 				switch(ElementToDraw->GetType())
 				{
-				case FVisualLogEntry::FElementToDraw::SinglePoint:
+				case EVisualLoggerShapeElement::SinglePoint:
 					{
 						const float Radius = float(ElementToDraw->Radius);
 						const bool bDrawLabel = ElementToDraw->Description.IsEmpty() == false;
@@ -2423,7 +2424,7 @@ void SLogVisualizer::DrawOnCanvas(UCanvas* Canvas, APlayerController*)
 						}
 					}
 					break;
-				case FVisualLogEntry::FElementToDraw::Segment:
+				case EVisualLoggerShapeElement::Segment:
 					{
 						const float Thickness = float(ElementToDraw->Thicknes);
 						const bool bDrawLabel = ElementToDraw->Description.IsEmpty() == false && ElementToDraw->Points.Num() > 2;
@@ -2453,7 +2454,7 @@ void SLogVisualizer::DrawOnCanvas(UCanvas* Canvas, APlayerController*)
 						}
 					}
 					break;
-				case FVisualLogEntry::FElementToDraw::Path:
+				case EVisualLoggerShapeElement::Path:
 					{
 						const float Thickness = float(ElementToDraw->Thicknes);
 						FVector Location = ElementToDraw->Points[0];
@@ -2467,7 +2468,7 @@ void SLogVisualizer::DrawOnCanvas(UCanvas* Canvas, APlayerController*)
 						}
 					}
 					break;
-				case FVisualLogEntry::FElementToDraw::Box:
+				case EVisualLoggerShapeElement::Box:
 					{
 						const float Thickness = float(ElementToDraw->Thicknes);
 						const bool bDrawLabel = ElementToDraw->Description.IsEmpty() == false && ElementToDraw->Points.Num() > 2;
@@ -2498,7 +2499,7 @@ void SLogVisualizer::DrawOnCanvas(UCanvas* Canvas, APlayerController*)
 						}
 					}
 					break;
-				case FVisualLogEntry::FElementToDraw::Cone:
+				case EVisualLoggerShapeElement::Cone:
 					{
 						const float Thickness = float(ElementToDraw->Thicknes);
 						const bool bDrawLabel = ElementToDraw->Description.IsEmpty() == false;
@@ -2516,7 +2517,7 @@ void SLogVisualizer::DrawOnCanvas(UCanvas* Canvas, APlayerController*)
 						}
 					}
 					break;
-				case FVisualLogEntry::FElementToDraw::Cylinder:
+				case EVisualLoggerShapeElement::Cylinder:
 					{
 						const float Thickness = float(ElementToDraw->Thicknes);
 						const bool bDrawLabel = ElementToDraw->Description.IsEmpty() == false;
@@ -2534,7 +2535,7 @@ void SLogVisualizer::DrawOnCanvas(UCanvas* Canvas, APlayerController*)
 						}
 					}
 					break;
-				case FVisualLogEntry::FElementToDraw::Capsule:
+				case EVisualLoggerShapeElement::Capsule:
 					{
 						const float Thickness = float(ElementToDraw->Thicknes);
 						const bool bDrawLabel = ElementToDraw->Description.IsEmpty() == false;
@@ -2614,7 +2615,7 @@ FReply SLogVisualizer::OnLoad()
 			LOCTEXT("OpenProjectBrowseTitle", "Open Project").ToString(),
 			LastBrowsePath,
 			TEXT(""),
-			LogVisualizer::FileTypes,
+			LogVisualizer::LoadFileTypes,
 			EFileDialogFlags::None,
 			OpenFilenames
 			);
@@ -2656,7 +2657,7 @@ FReply SLogVisualizer::OnSave()
 			LOCTEXT("NewProjectBrowseTitle", "Choose a project location").ToString(),
 			LastBrowsePath,
 			TEXT(""),
-			LogVisualizer::FileTypes,
+			LogVisualizer::SaveFileTypes,
 			EFileDialogFlags::None,
 			SaveFilenames
 			);
@@ -2827,38 +2828,74 @@ EColumnSortMode::Type SLogVisualizer::GetLogsSortMode() const
 
 void SLogVisualizer::LoadFiles(TArray<FString>& OpenFilenames)
 {
+	TArray<FVisualLogDevice::FVisualLogEntryItem> FrameCache;
+
 	for (int FilenameIndex = 0; FilenameIndex < OpenFilenames.Num(); ++FilenameIndex)
 	{
-		FArchive* FileAr = IFileManager::Get().CreateFileReader(*(OpenFilenames[FilenameIndex]));
-		if (FileAr != NULL)
+		FString CurrentFileName = OpenFilenames[FilenameIndex];
+		const bool bIsBinaryFile = CurrentFileName.Find(TEXT(".bvlog")) != INDEX_NONE;
+		if (!bIsBinaryFile)
 		{
-			TSharedPtr<FJsonObject> Object;
-			TSharedRef<TJsonReader<UCS2CHAR> > Reader = TJsonReader<UCS2CHAR>::Create(FileAr);
-
-			if (FJsonSerializer::Deserialize(Reader, Object))
+			FArchive* FileAr = IFileManager::Get().CreateFileReader(*CurrentFileName);
+			if (FileAr != NULL)
 			{
-				TArray< TSharedPtr<FJsonValue> > JsonLogs = Object->GetArrayField(VisualLogJson::TAG_LOGS);
-				for (int32 LogIndex = 0; LogIndex < JsonLogs.Num(); ++LogIndex)
+				TSharedPtr<FJsonObject> Object;
+				TSharedRef<TJsonReader<UCS2CHAR> > Reader = TJsonReader<UCS2CHAR>::Create(FileAr);
+
+				if (FJsonSerializer::Deserialize(Reader, Object))
 				{
-					TSharedPtr<FJsonObject> JsonLogObject = JsonLogs[LogIndex]->AsObject();
-					if (JsonLogObject.IsValid() != false)
+					TArray< TSharedPtr<FJsonValue> > JsonLogs = Object->GetArrayField(VisualLogJson::TAG_LOGS);
+					for (int32 LogIndex = 0; LogIndex < JsonLogs.Num(); ++LogIndex)
 					{
-						if (JsonLogObject->HasTypedField<EJson::String>(VisualLogJson::TAG_NAME))
+						TSharedPtr<FJsonObject> JsonLogObject = JsonLogs[LogIndex]->AsObject();
+						if (JsonLogObject.IsValid() != false)
 						{
-							TSharedPtr<FActorsVisLog> NewLog = MakeShareable(new FActorsVisLog(JsonLogs[LogIndex]));
-							LogVisualizer->AddLoadedLog(NewLog);
+							if (JsonLogObject->HasTypedField<EJson::String>(VisualLogJson::TAG_NAME))
+							{
+								TSharedPtr<FActorsVisLog> NewLog = MakeShareable(new FActorsVisLog(JsonLogs[LogIndex]));
+								LogVisualizer->AddLoadedLog(NewLog);
+							}
 						}
 					}
+					bIgnoreTrivialLogs = false;
 				}
-				bIgnoreTrivialLogs = false;
-			}
 
-			FileAr->Close();
+				FileAr->Close();
+			}
+		}
+		else
+		{
+			TArray<FVisualLogDevice::FVisualLogEntryItem> RecordedLogs;
+			FArchive* FileArchive = IFileManager::Get().CreateFileReader(*(OpenFilenames[FilenameIndex]));
+			FVisualLoggerHelpers::Serialize(*FileArchive, RecordedLogs);
+			FrameCache.Append(RecordedLogs);
+			FileArchive->Close();
+			delete FileArchive;
+			FileArchive = NULL;
 		}
 	}
 
 	if (OpenFilenames.Num() > 0)
 	{
+		if (FrameCache.Num() > 0)
+		{
+			TMap<FName, TSharedPtr<FActorsVisLog>> HelperMap;
+
+			for (auto& CurrentLog : FrameCache)
+			{
+				if (HelperMap.Contains(CurrentLog.OwnerName) == false)
+				{
+					HelperMap.Add(CurrentLog.OwnerName, MakeShareable(new FActorsVisLog(CurrentLog.OwnerName, NULL)));
+				}
+				HelperMap[CurrentLog.OwnerName]->Entries.Add(MakeShareable(new FVisualLogEntry(CurrentLog.Entry)));
+			}
+
+			for (auto& CurrentLog : HelperMap)
+			{
+				LogVisualizer->AddLoadedLog(CurrentLog.Value);
+			}
+		}
+
 		LogsList.Reset();
 		OutEntriesCached.Reset();
 		RebuildFilteredList();
@@ -2867,9 +2904,7 @@ void SLogVisualizer::LoadFiles(TArray<FString>& OpenFilenames)
 
 void SLogVisualizer::SaveSelectedLogs(FString& Filename)
 {
-	TSharedPtr<FJsonObject> Object = MakeShareable(new FJsonObject);
-
-	TArray< TSharedPtr<FJsonValue> > EntriesArray;
+	FVisualLogger::Get().Flush();
 	TArray< TSharedPtr<FLogsListItem> > ItemsToSave = LogsListWidget->GetSelectedItems();
 	if (ItemsToSave.Num() == 0)
 	{
@@ -2877,30 +2912,26 @@ void SLogVisualizer::SaveSelectedLogs(FString& Filename)
 		ItemsToSave = LogsList;
 	}
 
-	EntriesArray.Reserve(ItemsToSave.Num());
-
-
-	TSharedPtr<FLogsListItem>* LogListItem = ItemsToSave.GetData();
-	for (int32 ItemIndex = 0; ItemIndex < ItemsToSave.Num(); ++ItemIndex, ++LogListItem)
+	TArray<FVisualLogDevice::FVisualLogEntryItem> FrameCache;
+	for (auto CurrentItem : ItemsToSave)
 	{
-		if (LogListItem->IsValid() && LogVisualizer->Logs.IsValidIndex((*LogListItem)->LogIndex))
+		if (CurrentItem.IsValid() && LogVisualizer->Logs.IsValidIndex(CurrentItem->LogIndex))
 		{
-			TSharedPtr<FActorsVisLog> Log = LogVisualizer->Logs[(*LogListItem)->LogIndex];
-			EntriesArray.Add(Log->ToJson());
+			TSharedPtr<FActorsVisLog> Log = LogVisualizer->Logs[CurrentItem->LogIndex];
+			for (auto CurrentEntry : Log->Entries)
+			{
+				FrameCache.Add(FVisualLogDevice::FVisualLogEntryItem(Log->Name, *CurrentEntry.Get()));
+			}
 		}
 	}
-	
-	if (EntriesArray.Num() > 0)
-	{
-		Object->SetArrayField(VisualLogJson::TAG_LOGS, EntriesArray);
 
-		FArchive* FileAr = IFileManager::Get().CreateFileWriter(*Filename);
-		if (FileAr != NULL)
-		{
-			TSharedRef<TJsonWriter<UCS2CHAR> > Writer = TJsonWriter<UCS2CHAR>::Create(FileAr);
-			FJsonSerializer::Serialize( Object.ToSharedRef(), Writer );		
-			FileAr->Close();
-		}
+	if (FrameCache.Num())
+	{
+		FArchive* FileArchive = IFileManager::Get().CreateFileWriter(*Filename);
+		FVisualLoggerHelpers::Serialize(*FileArchive, FrameCache);
+		FileArchive->Close();
+		delete FileArchive;
+		FileArchive = NULL;
 	}
 }
 
diff --git a/Engine/Source/Developer/LogVisualizer/Private/SLogVisualizerReport.cpp b/Engine/Source/Developer/LogVisualizer/Private/SLogVisualizerReport.cpp
--- a/Engine/Source/Developer/LogVisualizer/Private/SLogVisualizerReport.cpp
+++ b/Engine/Source/Developer/LogVisualizer/Private/SLogVisualizerReport.cpp
@@ -63,19 +63,19 @@ SLogVisualizerReport::~SLogVisualizerReport()
 FText SLogVisualizerReport::GenerateReportText() const
 {
 	FString OutString;
-	TArray<FVisualLogEntry::FLogEvent> GlobalEventsStats;
+	TArray<FVisualLogEvent> GlobalEventsStats;
 	TMap<FString, TArray<FString> >	 EventToObjectsMap;
 
 
 	for (auto LogItem : SelectedItems)
 	{
-		TArray<FVisualLogEntry::FLogEvent> AllEvents;
+		TArray<FVisualLogEvent> AllEvents;
 
 		for (TSharedPtr<FVisualLogEntry> CurrentEntry : LogItem->Entries)
 		{
-			for (FVisualLogEntry::FLogEvent& Event : CurrentEntry->Events)
+			for (FVisualLogEvent& Event : CurrentEntry->Events)
 			{
-				int32 Index = AllEvents.Find(FVisualLogEntry::FLogEvent(Event));
+				int32 Index = AllEvents.Find(FVisualLogEvent(Event));
 				if (Index != INDEX_NONE)
 				{
 					for (auto& CurrentEventTag : Event.EventTags)
@@ -112,7 +112,7 @@ FText SLogVisualizerReport::GenerateReportText() const
 			OutString.Append( FString::Printf(TEXT("        %s occurred %d times\n\n"), *CurrentEvent.Name, CurrentEvent.Counter) );
 
 			bool bJustAdded = false;
-			int32 Index = GlobalEventsStats.Find(FVisualLogEntry::FLogEvent(CurrentEvent));
+			int32 Index = GlobalEventsStats.Find(FVisualLogEvent(CurrentEvent));
 			if (Index != INDEX_NONE)
 			{
 				GlobalEventsStats[Index].Counter += CurrentEvent.Counter;
@@ -125,7 +125,7 @@ FText SLogVisualizerReport::GenerateReportText() const
 				EventToObjectsMap.FindOrAdd(CurrentEvent.Name).Add(LogItem->Name.ToString());
 			}
 
-			Index = GlobalEventsStats.Find(FVisualLogEntry::FLogEvent(CurrentEvent));
+			Index = GlobalEventsStats.Find(FVisualLogEvent(CurrentEvent));
 			for (auto& CurrentEventTag : CurrentEvent.EventTags)
 			{
 				if (!bJustAdded)
diff --git a/Engine/Source/Runtime/AIModule/Classes/VisualLoggerExtension.h b/Engine/Source/Runtime/AIModule/Classes/VisualLoggerExtension.h
--- a/Engine/Source/Runtime/AIModule/Classes/VisualLoggerExtension.h
+++ b/Engine/Source/Runtime/AIModule/Classes/VisualLoggerExtension.h
@@ -1,7 +1,7 @@
 // Copyright 1998-2014 Epic Games, Inc. All Rights Reserved.
 #pragma once
 #if ENABLE_VISUAL_LOG
-#	include "VisualLog.h"
+#include "VisualLogger/VisualLogger.h"
 #endif
 #include "VisualLoggerExtension.generated.h"
 
@@ -16,8 +16,8 @@ class FVisualLoggerExtension : public FVisualLogExtensionInterface
 public:
 	FVisualLoggerExtension();
 	virtual void OnTimestampChange(float Timestamp, class UWorld* InWorld, class AActor* HelperActor) override;
-	virtual void DrawData(class UWorld* InWorld, class UCanvas* Canvas, class AActor* HelperActor, const FName& TagName, const struct FVisualLogEntry::FDataBlock& DataBlock, float Timestamp) override;
-	virtual void DisableDrawingForData(class UWorld* InWorld, class UCanvas* Canvas, class AActor* HelperActor, const FName& TagName, const FVisualLogEntry::FDataBlock& DataBlock, float Timestamp) override;
+	virtual void DrawData(class UWorld* InWorld, class UCanvas* Canvas, class AActor* HelperActor, const FName& TagName, const struct FVisualLogDataBlock& DataBlock, float Timestamp) override;
+	virtual void DisableDrawingForData(class UWorld* InWorld, class UCanvas* Canvas, class AActor* HelperActor, const FName& TagName, const FVisualLogDataBlock& DataBlock, float Timestamp) override;
 	virtual void LogEntryLineSelectionChanged(TSharedPtr<struct FLogEntryItem> SelectedItem, int64 UserData, FName TagName) override;
 
 private:
diff --git a/Engine/Source/Runtime/AIModule/Private/AIController.cpp b/Engine/Source/Runtime/AIModule/Private/AIController.cpp
--- a/Engine/Source/Runtime/AIModule/Private/AIController.cpp
+++ b/Engine/Source/Runtime/AIModule/Private/AIController.cpp
@@ -111,7 +111,7 @@ void AAIController::DisplayDebug(UCanvas* Canvas, const FDebugDisplayInfo& Debug
 
 void AAIController::GrabDebugSnapshot(FVisualLogEntry* Snapshot) const
 {
-	FVisualLogEntry::FStatusCategory MyCategory;
+	FVisualLogStatusCategory MyCategory;
 	MyCategory.Category = TEXT("AI Controller");
 	MyCategory.Add(TEXT("Pawn"), GetNameSafe(GetPawn()));
 	AActor* FocusActor = GetFocusActor();
diff --git a/Engine/Source/Runtime/AIModule/Private/AIModulePrivate.h b/Engine/Source/Runtime/AIModule/Private/AIModulePrivate.h
--- a/Engine/Source/Runtime/AIModule/Private/AIModulePrivate.h
+++ b/Engine/Source/Runtime/AIModule/Private/AIModulePrivate.h
@@ -59,7 +59,7 @@
 #include "Perception/AISense.h"
 #include "Perception/AIPerceptionComponent.h"
 
-#include "VisualLog.h"
+#include "VisualLogger/VisualLogger.h"
 
 DECLARE_LOG_CATEGORY_EXTERN(LogAINavigation, Warning, All);
 DECLARE_LOG_CATEGORY_EXTERN(LogBehaviorTree, Display, All);
diff --git a/Engine/Source/Runtime/AIModule/Private/Actions/PawnActionsComponent.cpp b/Engine/Source/Runtime/AIModule/Private/Actions/PawnActionsComponent.cpp
--- a/Engine/Source/Runtime/AIModule/Private/Actions/PawnActionsComponent.cpp
+++ b/Engine/Source/Runtime/AIModule/Private/Actions/PawnActionsComponent.cpp
@@ -521,7 +521,7 @@ void UPawnActionsComponent::DescribeSelfToVisLog(FVisualLogEntry* Snapshot) cons
 			continue;
 		}
 
-		FVisualLogEntry::FStatusCategory StatusCategory;
+		FVisualLogStatusCategory StatusCategory;
 		StatusCategory.Category = Category + TEXT(": ") + GetPriorityName(PriorityIndex);
 
 		while (Action)
diff --git a/Engine/Source/Runtime/AIModule/Private/BehaviorTree/BehaviorTreeComponent.cpp b/Engine/Source/Runtime/AIModule/Private/BehaviorTree/BehaviorTreeComponent.cpp
--- a/Engine/Source/Runtime/AIModule/Private/BehaviorTree/BehaviorTreeComponent.cpp
+++ b/Engine/Source/Runtime/AIModule/Private/BehaviorTree/BehaviorTreeComponent.cpp
@@ -1622,7 +1622,7 @@ void UBehaviorTreeComponent::DescribeSelfToVisLog(FVisualLogEntry* Snapshot) con
 		const FBehaviorTreeInstance& InstanceInfo = InstanceStack[InstanceIndex];
 		const FBehaviorTreeInstanceId& InstanceId = KnownInstances[InstanceInfo.InstanceIdIndex];
 
-		FVisualLogEntry::FStatusCategory StatusCategory;
+		FVisualLogStatusCategory StatusCategory;
 		StatusCategory.Category = FString::Printf(TEXT("BehaviorTree %d (asset: %s)"), InstanceIndex, *GetNameSafe(InstanceId.TreeAsset));
 		
 		TArray<FString> Descriptions;
diff --git a/Engine/Source/Runtime/AIModule/Private/BehaviorTree/BlackboardComponent.cpp b/Engine/Source/Runtime/AIModule/Private/BehaviorTree/BlackboardComponent.cpp
--- a/Engine/Source/Runtime/AIModule/Private/BehaviorTree/BlackboardComponent.cpp
+++ b/Engine/Source/Runtime/AIModule/Private/BehaviorTree/BlackboardComponent.cpp
@@ -316,7 +316,7 @@ FString UBlackboardComponent::DescribeKeyValue(FBlackboard::FKey KeyID, EBlackbo
 
 void UBlackboardComponent::DescribeSelfToVisLog(FVisualLogEntry* Snapshot) const
 {
-	FVisualLogEntry::FStatusCategory Category;
+	FVisualLogStatusCategory Category;
 	Category.Category = FString::Printf(TEXT("Blackboard (asset: %s)"), *GetNameSafe(BlackboardAsset));
 
 	for (UBlackboardData* It = BlackboardAsset; It; It = It->Parent)
diff --git a/Engine/Source/Runtime/AIModule/Private/BrainComponent.cpp b/Engine/Source/Runtime/AIModule/Private/BrainComponent.cpp
--- a/Engine/Source/Runtime/AIModule/Private/BrainComponent.cpp
+++ b/Engine/Source/Runtime/AIModule/Private/BrainComponent.cpp
@@ -193,7 +193,7 @@ void UBrainComponent::DescribeSelfToVisLog(FVisualLogEntry* Snapshot) const
 		return;
 	}
 
-	FVisualLogEntry::FStatusCategory StatusCategory;
+	FVisualLogStatusCategory StatusCategory;
 	StatusCategory.Category = FString::Printf(TEXT("Resource lock: %s"), *ResourceLock.GetLockSourceName());
 	for (int32 LockLevel = 0; LockLevel < int32(EAILockSource::MAX); ++LockLevel)
 	{
diff --git a/Engine/Source/Runtime/AIModule/Private/EnvironmentQuery/EnvQueryDebugHelpers.cpp b/Engine/Source/Runtime/AIModule/Private/EnvironmentQuery/EnvQueryDebugHelpers.cpp
--- a/Engine/Source/Runtime/AIModule/Private/EnvironmentQuery/EnvQueryDebugHelpers.cpp
+++ b/Engine/Source/Runtime/AIModule/Private/EnvironmentQuery/EnvQueryDebugHelpers.cpp
@@ -138,7 +138,7 @@ void UEnvQueryDebugHelpers::LogQueryInternal(FEnvQueryInstance& Query, const FLo
 	TArray<uint8> BlobArray;
 	UEnvQueryDebugHelpers::QueryToBlobArray(Query, BlobArray);
 	const FString AdditionalLogInfo = FString::Printf(TEXT("Executed EQS: \n - Name: '%s' (id=%d, option=%d),\n - All Items: %d,\n - ValidItems: %d"), *Query.QueryName, Query.QueryID, Query.OptionIndex, Query.ItemDetails.Num(), Query.NumValidItems);
-	FVisualLogEntry::FLogLine Line(Category.GetCategoryName(), Verbosity, AdditionalLogInfo, Query.QueryID);
+	FVisualLogLine Line(Category.GetCategoryName(), Verbosity, AdditionalLogInfo, Query.QueryID);
 	Line.TagName = *EVisLogTags::TAG_EQS;
 	Line.UniqueId = UniqueId;
 
diff --git a/Engine/Source/Runtime/AIModule/Private/EnvironmentQuery/Generators/EnvQueryGenerator_BlueprintBase.cpp b/Engine/Source/Runtime/AIModule/Private/EnvironmentQuery/Generators/EnvQueryGenerator_BlueprintBase.cpp
--- a/Engine/Source/Runtime/AIModule/Private/EnvironmentQuery/Generators/EnvQueryGenerator_BlueprintBase.cpp
+++ b/Engine/Source/Runtime/AIModule/Private/EnvironmentQuery/Generators/EnvQueryGenerator_BlueprintBase.cpp
@@ -7,7 +7,7 @@
 #include "EnvironmentQuery/Generators/EnvQueryGenerator_BlueprintBase.h"
 #include "EnvironmentQuery/Items/EnvQueryItemType_Actor.h"
 #include "AI/Navigation/RecastNavMesh.h"
-#include "VisualLog.h"
+#include "VisualLogger/VisualLogger.h"
 
 #define LOCTEXT_NAMESPACE "EnvQueryGenerator"
 
diff --git a/Engine/Source/Runtime/AIModule/Private/Navigation/CrowdFollowingComponent.cpp b/Engine/Source/Runtime/AIModule/Private/Navigation/CrowdFollowingComponent.cpp
--- a/Engine/Source/Runtime/AIModule/Private/Navigation/CrowdFollowingComponent.cpp
+++ b/Engine/Source/Runtime/AIModule/Private/Navigation/CrowdFollowingComponent.cpp
@@ -867,7 +867,7 @@ void UCrowdFollowingComponent::DescribeSelfToVisLog(FVisualLogEntry* Snapshot) c
 		return;
 	}
 
-	FVisualLogEntry::FStatusCategory Category;
+	FVisualLogStatusCategory Category;
 	Category.Category = TEXT("Path following");
 
 	if (DestinationActor.IsValid())
diff --git a/Engine/Source/Runtime/AIModule/Private/Navigation/PathFollowingComponent.cpp b/Engine/Source/Runtime/AIModule/Private/Navigation/PathFollowingComponent.cpp
--- a/Engine/Source/Runtime/AIModule/Private/Navigation/PathFollowingComponent.cpp
+++ b/Engine/Source/Runtime/AIModule/Private/Navigation/PathFollowingComponent.cpp
@@ -1308,7 +1308,7 @@ void UPathFollowingComponent::DisplayDebug(UCanvas* Canvas, const FDebugDisplayI
 #if ENABLE_VISUAL_LOG
 void UPathFollowingComponent::DescribeSelfToVisLog(FVisualLogEntry* Snapshot) const
 {
-	FVisualLogEntry::FStatusCategory Category;
+	FVisualLogStatusCategory Category;
 	Category.Category = TEXT("Path following");
 
 	if (DestinationActor.IsValid())
diff --git a/Engine/Source/Runtime/AIModule/Private/VisualLoggerExtension.cpp b/Engine/Source/Runtime/AIModule/Private/VisualLoggerExtension.cpp
--- a/Engine/Source/Runtime/AIModule/Private/VisualLoggerExtension.cpp
+++ b/Engine/Source/Runtime/AIModule/Private/VisualLoggerExtension.cpp
@@ -27,7 +27,7 @@ void FVisualLoggerExtension::OnTimestampChange(float Timestamp, UWorld* InWorld,
 #endif
 }
 
-void FVisualLoggerExtension::DisableDrawingForData(UWorld* InWorld, UCanvas* Canvas, AActor* HelperActor, const FName& TagName, const FVisualLogEntry::FDataBlock& DataBlock, float Timestamp)
+void FVisualLoggerExtension::DisableDrawingForData(UWorld* InWorld, UCanvas* Canvas, AActor* HelperActor, const FName& TagName, const FVisualLogDataBlock& DataBlock, float Timestamp)
 {
 #if USE_EQS_DEBUGGER
 	if (TagName == *EVisLogTags::TAG_EQS && CurrentTimestamp == Timestamp)
@@ -65,7 +65,7 @@ void FVisualLoggerExtension::LogEntryLineSelectionChanged(TSharedPtr<FLogEntryIt
 
 extern RENDERCORE_API FTexture* GWhiteTexture;
 
-void FVisualLoggerExtension::DrawData(UWorld* InWorld, UCanvas* Canvas, AActor* HelperActor, const FName& TagName, const FVisualLogEntry::FDataBlock& DataBlock, float Timestamp)
+void FVisualLoggerExtension::DrawData(UWorld* InWorld, UCanvas* Canvas, AActor* HelperActor, const FName& TagName, const FVisualLogDataBlock& DataBlock, float Timestamp)
 {
 #if USE_EQS_DEBUGGER
 	if (TagName == *EVisLogTags::TAG_EQS && HelperActor)
diff --git a/Engine/Source/Runtime/Engine/Classes/VisualLogger/VisualLoggerBinaryFileDevice.h b/Engine/Source/Runtime/Engine/Classes/VisualLogger/VisualLoggerBinaryFileDevice.h
--- a/Engine/Source/Runtime/Engine/Classes/VisualLogger/VisualLoggerBinaryFileDevice.h
+++ b/Engine/Source/Runtime/Engine/Classes/VisualLogger/VisualLoggerBinaryFileDevice.h
@@ -3,6 +3,9 @@
 #include "VisualLogger/VisualLogger.h"
 
 #if ENABLE_VISUAL_LOG
+
+#define VISLOG_FILENAME_EXT TEXT("bvlog")
+
 class FVisualLoggerBinaryFileDevice : public FVisualLogDevice
 {
 public:
@@ -18,15 +21,9 @@ public:
 	virtual void StopRecordingToFile(float TImeStamp) override;
 	virtual void SetFileName(const FString& InFileName) override;
 	virtual void Serialize(const class UObject* LogOwner, FName OwnerName, const FVisualLogEntry& LogEntry) override;
+	virtual void GetRecordedLogs(TArray<FVisualLogEntryItem>& RecordedLogs) const override { RecordedLogs = FrameCache; }
+	virtual bool HasFlags(int32 InFlags) const { return !!(InFlags & (EVisualLoggerDeviceFlags::CanSaveToFile | EVisualLoggerDeviceFlags::StoreLogsLocally)); }
 
-	struct FFrameCacheItem
-	{
-		FFrameCacheItem() {}
-		FFrameCacheItem(FName InOwnerName, const FVisualLogEntry& LogEntry) : OwnerName(InOwnerName), Entry(LogEntry) { }
-
-		FName OwnerName;
-		FVisualLogEntry Entry;
-	};
 protected:
 	int32 bUseCompression : 1;
 	float FrameCacheLenght;
@@ -35,6 +32,6 @@ protected:
 	FArchive* FileArchive;
 	FString TempFileName;
 	FString FileName;
-	TArray<FFrameCacheItem> FrameCache;
+	TArray<FVisualLogEntryItem> FrameCache;
 };
 #endif
diff --git a/Engine/Source/Runtime/Engine/Private/AI/Navigation/NavLinkCustomComponent.cpp b/Engine/Source/Runtime/Engine/Private/AI/Navigation/NavLinkCustomComponent.cpp
--- a/Engine/Source/Runtime/Engine/Private/AI/Navigation/NavLinkCustomComponent.cpp
+++ b/Engine/Source/Runtime/Engine/Private/AI/Navigation/NavLinkCustomComponent.cpp
@@ -1,7 +1,7 @@
 // Copyright 1998-2014 Epic Games, Inc. All Rights Reserved.
 
 #include "EnginePrivate.h"
-#include "VisualLog.h"
+#include "VisualLogger/VisualLogger.h"
 #include "AI/Navigation/NavAreas/NavArea_Null.h"
 #include "AI/Navigation/NavAreas/NavArea_Default.h"
 // @todo to be addressed when removing AIModule circular dependency
diff --git a/Engine/Source/Runtime/Engine/Private/AI/Navigation/NavigationPath.cpp b/Engine/Source/Runtime/Engine/Private/AI/Navigation/NavigationPath.cpp
--- a/Engine/Source/Runtime/Engine/Private/AI/Navigation/NavigationPath.cpp
+++ b/Engine/Source/Runtime/Engine/Private/AI/Navigation/NavigationPath.cpp
@@ -1,7 +1,7 @@
 // Copyright 1998-2014 Epic Games, Inc. All Rights Reserved.
 
 #include "EnginePrivate.h"
-#include "VisualLog.h"
+#include "VisualLogger/VisualLogger.h"
 #include "AI/Navigation/NavigationTypes.h"
 #include "AI/Navigation/RecastNavMesh.h"
 #include "AI/Navigation/NavigationPath.h"
@@ -267,7 +267,7 @@ bool FNavigationPath::DoesIntersectBox(const FBox& Box, uint32 StartingIndex, in
 void FNavigationPath::DescribeSelfToVisLog(FVisualLogEntry* Snapshot) const 
 {
 	const int32 NumPathVerts = PathPoints.Num();
-	FVisualLogEntry::FElementToDraw Element(FVisualLogEntry::FElementToDraw::Path);
+	FVisualLogShapeElement Element(EVisualLoggerShapeElement::Path);
 	Element.Category = LogNavigation.GetCategoryName();
 	Element.SetColor(FColorList::Green);
 	Element.Points.Reserve(NumPathVerts);
@@ -909,7 +909,7 @@ void FNavMeshPath::DescribeSelfToVisLog(FVisualLogEntry* Snapshot) const
 
 	// draw corridor
 #if WITH_RECAST
-	FVisualLogEntry::FElementToDraw CorridorElem(FVisualLogEntry::FElementToDraw::Segment);
+	FVisualLogShapeElement CorridorElem(EVisualLoggerShapeElement::Segment);
 	CorridorElem.SetColor(FColorList::Cyan);
 	CorridorElem.Category = LogNavigation.GetCategoryName();
 	CorridorElem.Points.Reserve(PathCorridor.Num() * 6);
@@ -940,7 +940,7 @@ void FNavMeshPath::DescribeSelfToVisLog(FVisualLogEntry* Snapshot) const
 		const UClass* AreaClass = NavMesh->GetAreaClass(AreaID);
 		if (AreaClass && AreaClass != UNavigationSystem::GetDefaultWalkableArea())
 		{
-			FVisualLogEntry::FElementToDraw AreaMarkElem(FVisualLogEntry::FElementToDraw::Segment);
+			FVisualLogShapeElement AreaMarkElem(EVisualLoggerShapeElement::Segment);
 			AreaMarkElem.SetColor(FColorList::Orange);
 			AreaMarkElem.Category = LogNavigation.GetCategoryName();
 			AreaMarkElem.Thicknes = 2;
diff --git a/Engine/Source/Runtime/Engine/Private/AI/Navigation/NavigationSystemHelpers.cpp b/Engine/Source/Runtime/Engine/Private/AI/Navigation/NavigationSystemHelpers.cpp
--- a/Engine/Source/Runtime/Engine/Private/AI/Navigation/NavigationSystemHelpers.cpp
+++ b/Engine/Source/Runtime/Engine/Private/AI/Navigation/NavigationSystemHelpers.cpp
@@ -7,7 +7,7 @@
 #include "RecastNavMeshGenerator.h"
 #endif // WITH_RECAST
 #include "AI/NavigationSystemHelpers.h"
-#include "VisualLog.h"
+#include "VisualLogger/VisualLogger.h"
 
 namespace NavigationHelper
 {
diff --git a/Engine/Source/Runtime/Engine/Private/AI/Navigation/PImplRecastNavMesh.cpp b/Engine/Source/Runtime/Engine/Private/AI/Navigation/PImplRecastNavMesh.cpp
--- a/Engine/Source/Runtime/Engine/Private/AI/Navigation/PImplRecastNavMesh.cpp
+++ b/Engine/Source/Runtime/Engine/Private/AI/Navigation/PImplRecastNavMesh.cpp
@@ -17,7 +17,7 @@
 #include "RecastNavMeshGenerator.h"
 
 #include "AI/Navigation/NavLinkCustomInterface.h"
-#include "VisualLog.h"
+#include "VisualLogger/VisualLogger.h"
 
 
 //----------------------------------------------------------------------//
diff --git a/Engine/Source/Runtime/Engine/Private/AI/Navigation/RecastNavMesh.cpp b/Engine/Source/Runtime/Engine/Private/AI/Navigation/RecastNavMesh.cpp
--- a/Engine/Source/Runtime/Engine/Private/AI/Navigation/RecastNavMesh.cpp
+++ b/Engine/Source/Runtime/Engine/Private/AI/Navigation/RecastNavMesh.cpp
@@ -7,7 +7,7 @@
 #include "AI/Navigation/NavAreas/NavArea_LowHeight.h"
 #include "AI/Navigation/NavLinkCustomInterface.h"
 #include "AI/Navigation/RecastNavMesh.h"
-#include "VisualLog.h"
+#include "VisualLogger/VisualLogger.h"
 
 #if WITH_EDITOR
 #include "ObjectEditorUtils.h"
diff --git a/Engine/Source/Runtime/Engine/Private/Actor.cpp b/Engine/Source/Runtime/Engine/Private/Actor.cpp
--- a/Engine/Source/Runtime/Engine/Private/Actor.cpp
+++ b/Engine/Source/Runtime/Engine/Private/Actor.cpp
@@ -14,7 +14,7 @@
 #include "Matinee/InterpGroup.h"
 #include "Matinee/InterpGroupInst.h"
 #include "Particles/ParticleSystemComponent.h"
-#include "VisualLog.h"
+#include "VisualLogger/VisualLogger.h"
 #include "Animation/AnimInstance.h"
 #include "Engine/DemoNetDriver.h"
 
diff --git a/Engine/Source/Runtime/Engine/Private/CheatManager.cpp b/Engine/Source/Runtime/Engine/Private/CheatManager.cpp
--- a/Engine/Source/Runtime/Engine/Private/CheatManager.cpp
+++ b/Engine/Source/Runtime/Engine/Private/CheatManager.cpp
@@ -6,7 +6,7 @@
 #include "SlateReflector.h"
 #include "NavDataGenerator.h"
 #include "OnlineSubsystemUtils.h"
-#include "VisualLog.h"
+#include "VisualLogger/VisualLogger.h"
 #include "GameFramework/Character.h"
 
 #if WITH_EDITOR
diff --git a/Engine/Source/Runtime/Engine/Private/Controller.cpp b/Engine/Source/Runtime/Engine/Private/Controller.cpp
--- a/Engine/Source/Runtime/Engine/Private/Controller.cpp
+++ b/Engine/Source/Runtime/Engine/Private/Controller.cpp
@@ -8,7 +8,7 @@
 #include "Net/UnrealNetwork.h"
 #include "ConfigCacheIni.h"
 #include "NetworkingDistanceConstants.h"
-#include "VisualLog.h"
+#include "VisualLogger/VisualLogger.h"
 #include "GameFramework/Character.h"
 #include "GameFramework/PawnMovementComponent.h"
 #include "GameFramework/CharacterMovementComponent.h"
diff --git a/Engine/Source/Runtime/Engine/Private/Debug/LogVisualizerHUD.cpp b/Engine/Source/Runtime/Engine/Private/Debug/LogVisualizerHUD.cpp
--- a/Engine/Source/Runtime/Engine/Private/Debug/LogVisualizerHUD.cpp
+++ b/Engine/Source/Runtime/Engine/Private/Debug/LogVisualizerHUD.cpp
@@ -3,7 +3,7 @@
 #include "EnginePrivate.h"
 #include "Debug/LogVisualizerHUD.h"
 #include "Debug/LogVisualizerCameraController.h"
-#include "VisualLog.h"
+#include "VisualLogger/VisualLogger.h"
 
 #define LOCTEXT_NAMESPACE "ALogVisualizerHUD"
 
diff --git a/Engine/Source/Runtime/Engine/Private/Pawn.cpp b/Engine/Source/Runtime/Engine/Private/Pawn.cpp
--- a/Engine/Source/Runtime/Engine/Private/Pawn.cpp
+++ b/Engine/Source/Runtime/Engine/Private/Pawn.cpp
@@ -14,7 +14,7 @@
 #include "ParticleDefinitions.h"
 #include "DisplayDebugHelpers.h"
 #include "NetworkingDistanceConstants.h"
-#include "VisualLog.h"
+#include "VisualLogger/VisualLogger.h"
 #include "Engine/InputDelegateBinding.h"
 
 DEFINE_LOG_CATEGORY(LogDamage);
diff --git a/Engine/Source/Runtime/Engine/Private/UnrealEngine.cpp b/Engine/Source/Runtime/Engine/Private/UnrealEngine.cpp
--- a/Engine/Source/Runtime/Engine/Private/UnrealEngine.cpp
+++ b/Engine/Source/Runtime/Engine/Private/UnrealEngine.cpp
@@ -9,7 +9,7 @@
 #include "Engine/ObjectReferencer.h"
 #include "Net/UnrealNetwork.h"
 #include "Engine/Console.h"
-#include "VisualLog.h"
+#include "VisualLogger/VisualLogger.h"
 #include "FileManagerGeneric.h"
 #include "Database.h"
 #include "SkeletalMeshMerge.h"
diff --git a/Engine/Source/Runtime/Engine/Private/VisualLog.cpp b/Engine/Source/Runtime/Engine/Private/VisualLog.cpp
deleted file mode 100644
diff --git a/Engine/Source/Runtime/Engine/Private/VisualLogger/VisualLogger.cpp b/Engine/Source/Runtime/Engine/Private/VisualLogger/VisualLogger.cpp
--- a/Engine/Source/Runtime/Engine/Private/VisualLogger/VisualLogger.cpp
+++ b/Engine/Source/Runtime/Engine/Private/VisualLogger/VisualLogger.cpp
@@ -1,30 +1,21 @@
 // Copyright 1998-2014 Epic Games, Inc. All Rights Reserved.
 
 #include "EnginePrivate.h"
-#include "VisualLog.h"
+#include "VisualLogger/VisualLogger.h"
 #include "VisualLogger/VisualLogger.h"
 #include "VisualLogger/VisualLoggerBinaryFileDevice.h"
 
 #if ENABLE_VISUAL_LOG 
 
+DEFINE_STAT(STAT_VisualLog);
 DEFINE_LOG_CATEGORY(LogVisual);
 
-TMap<class UObject*, TArray<TWeakObjectPtr<const class UObject> > > FVisualLogger::RedirectionMap;
+TMap<UObject*, TArray<TWeakObjectPtr<const UObject> > > FVisualLogger::RedirectionMap;
 
 FVisualLogger::FVisualLogger()
 {
-	UseBinaryFileDevice = false;
-	GConfig->GetBool(TEXT("VisualLogger"), TEXT("UseBinaryFileDevice"), UseBinaryFileDevice, GEngineIni);
-
 	BlockAllCategories(false);
-	if (UseBinaryFileDevice)
-	{
-		AddDevice(&FVisualLoggerBinaryFileDevice::Get());
-	}
-	else
-	{
-		AddDevice(&FVisualLog::GetStatic());
-	}
+	AddDevice(&FVisualLoggerBinaryFileDevice::Get());
 	SetIsRecording(GEngine ? !!GEngine->bEnableVisualLogRecordingOnStart : false);
 	SetIsRecordingOnServer(false);
 
@@ -35,22 +26,15 @@ FVisualLogger::FVisualLogger()
 	}
 }
 
-FVisualLogger::~FVisualLogger()
+void FVisualLogger::Shutdown()
 {
+	SetIsRecording(false);
+	SetIsRecordingToFile(false);
+
 	if (UseBinaryFileDevice)
 	{
 		RemoveDevice(&FVisualLoggerBinaryFileDevice::Get());
 	}
-	else
-	{
-		RemoveDevice(&FVisualLog::GetStatic());
-	}
-}
-
-void FVisualLogger::Shutdown()
-{
-	SetIsRecording(false);
-	SetIsRecordingToFile(false);
 }
 
 void FVisualLogger::Cleanup(bool bReleaseMemory)
@@ -82,7 +66,7 @@ int32 FVisualLogger::GetUniqueId(float Timestamp)
 	return LastUniqueIds.FindOrAdd(Timestamp)++;
 }
 
-void FVisualLogger::Redirect(class UObject* FromObject, class UObject* ToObject)
+void FVisualLogger::Redirect(UObject* FromObject, UObject* ToObject)
 {
 	if (FromObject == ToObject)
 	{
@@ -103,7 +87,7 @@ void FVisualLogger::Redirect(class UObject* FromObject, class UObject* ToObject)
 	UE_CVLOG(FromObject != NULL, FromObject, LogVisual, Log, TEXT("Redirected '%s' to '%s'"), *FromObject->GetName(), *NewRedierection->GetName());
 }
 
-class UObject* FVisualLogger::FindRedirection(const class UObject* Object)
+class UObject* FVisualLogger::FindRedirection(const UObject* Object)
 {
 	if (RedirectionMap.Contains(Object) == false)
 	{
@@ -144,7 +128,7 @@ void FVisualLogger::SetIsRecordingToFile(bool InIsRecording)
 	{
 		for (auto* Device : OutputDevices)
 		{
-			if (Device->HasFlags(VisualLogger::CanSaveToFile))
+			if (Device->HasFlags(EVisualLoggerDeviceFlags::CanSaveToFile))
 			{
 				Device->SetFileName(OutputFileName);
 				Device->StopRecordingToFile(GWorld.GetReference() ? GWorld->TimeSeconds : StartRecordingToFileTime);
@@ -156,7 +140,7 @@ void FVisualLogger::SetIsRecordingToFile(bool InIsRecording)
 		StartRecordingToFileTime = GWorld.GetReference() ? GWorld->TimeSeconds : 0;
 		for (auto* Device : OutputDevices)
 		{
-			if (Device->HasFlags(VisualLogger::CanSaveToFile))
+			if (Device->HasFlags(EVisualLoggerDeviceFlags::CanSaveToFile))
 			{
 				Device->StartRecordingToFile(StartRecordingToFileTime);
 			}
@@ -166,4 +150,61 @@ void FVisualLogger::SetIsRecordingToFile(bool InIsRecording)
 	bIsRecordingToFile = InIsRecording;
 }
 
-#endif //ENABLE_VISUAL_LOG
\ No newline at end of file
+#endif //ENABLE_VISUAL_LOG
+
+const FGuid EVisualLoggerVersion::GUID = FGuid(0xA4237A36, 0xCAEA41C9, 0x8FA218F8, 0x58681BF3);
+FCustomVersionRegistration GVisualLoggerVersion(EVisualLoggerVersion::GUID, EVisualLoggerVersion::LatestVersion, TEXT("VisualLogger"));
+
+#if !(UE_BUILD_SHIPPING || UE_BUILD_TEST)
+
+#include "Developer/LogVisualizer/Public/LogVisualizerModule.h"
+
+static class FLogVisualizerExec : private FSelfRegisteringExec
+{
+public:
+	/** Console commands, see embeded usage statement **/
+	virtual bool Exec(UWorld* InWorld, const TCHAR* Cmd, FOutputDevice& Ar) override
+	{
+		if (FParse::Command(&Cmd, TEXT("VISLOG")))
+		{
+#if ENABLE_VISUAL_LOG
+			FString Command = FParse::Token(Cmd, 0);
+			if (Command == TEXT("record"))
+			{
+				FVisualLogger::Get().SetIsRecording(true);
+				return true;
+			}
+			else if (Command == TEXT("stop"))
+			{
+				FVisualLogger::Get().SetIsRecording(false);
+				return true;
+			}
+			else if (Command == TEXT("disableallbut"))
+			{
+				FString Category = FParse::Token(Cmd, 1);
+				FVisualLogger::Get().BlockAllCategories(true);
+				FVisualLogger::Get().GetWhiteList().AddUnique(*Category);
+				return true;
+			}
+#if WITH_EDITOR
+			else if (Command == TEXT("exit"))
+			{
+				FLogVisualizerModule::Get()->CloseUI(InWorld);
+				return true;
+			}
+			else
+			{
+				FLogVisualizerModule::Get()->SummonUI(InWorld);
+				return true;
+			}
+#endif
+#else
+			UE_LOG(LogVisual, Warning, TEXT("Unable to open LogVisualizer - logs are disabled"));
+#endif
+		}
+		return false;
+	}
+} LogVisualizerExec;
+
+
+#endif // !(UE_BUILD_SHIPPING || UE_BUILD_TEST)
diff --git a/Engine/Source/Runtime/Engine/Private/VisualLogger/VisualLoggerAutomationTests.cpp b/Engine/Source/Runtime/Engine/Private/VisualLogger/VisualLoggerAutomationTests.cpp
--- a/Engine/Source/Runtime/Engine/Private/VisualLogger/VisualLoggerAutomationTests.cpp
+++ b/Engine/Source/Runtime/Engine/Private/VisualLogger/VisualLoggerAutomationTests.cpp
@@ -2,7 +2,6 @@
 
 #include "EnginePrivate.h"
 
-#include "VisualLog.h"
 #include "VisualLogger/VisualLogger.h"
 #include "VisualLogger/VisualLoggerAutomationTests.h"
 #include "VisualLogger/VisualLoggerBinaryFileDevice.h"
@@ -20,9 +19,6 @@ class FVisualLoggerTestDevice : public FVisualLogDevice
 public:
 	FVisualLoggerTestDevice();
 	virtual void Cleanup(bool bReleaseMemory = false) override;
-	virtual void StartRecordingToFile(float TImeStamp) override;
-	virtual void StopRecordingToFile(float TImeStamp) override;
-	virtual void SetFileName(const FString& InFileName) override;
 	virtual void Serialize(const class UObject* LogOwner, FName OwnerName, const FVisualLogEntry& LogEntry) override;
 
 	class UObject* LastObject;
@@ -40,22 +36,7 @@ void FVisualLoggerTestDevice::Cleanup(bool bReleaseMemory)
 	LastEntry.Reset();
 }
 
-void FVisualLoggerTestDevice::StartRecordingToFile(float TImeStamp)
-{
-
-}
-
-void FVisualLoggerTestDevice::StopRecordingToFile(float TImeStamp)
-{
-
-}
-
-void FVisualLoggerTestDevice::SetFileName(const FString& InFileName)
-{
-
-}
-
-void FVisualLoggerTestDevice::Serialize(const class UObject* LogOwner, FName OwnerName, const FVisualLogEntry& LogEntry)
+void FVisualLoggerTestDevice::Serialize(const UObject* LogOwner, FName OwnerName, const FVisualLogEntry& LogEntry)
 {
 	LastObject = const_cast<class UObject*>(LogOwner);
 	LastEntry = LogEntry;
@@ -120,10 +101,10 @@ bool FVisualLogTest::RunTest(const FString& Parameters)
 
 	{
 		const FString TextToLog = TEXT("Hello World");
-		UE_VLOG(GWorld, LogVisual, Log, *TextToLog);
+		UE_VLOG_UELOG(GWorld, LogVisual, Warning, TEXT("%s"), *TextToLog);
 		CHECK_SUCCESS(Context.Device.LastObject != GWorld);
 		CHECK_SUCCESS(Context.Device.LastEntry.TimeStamp == -1);
-		FVisualLogEntry* CurrentEntry = FVisualLogger::Get().GetEntryToWrite(GWorld, GWorld->TimeSeconds, VisualLogger::DontCreate);
+		FVisualLogEntry* CurrentEntry = FVisualLogger::Get().GetEntryToWrite(GWorld, GWorld->TimeSeconds, ECreateIfNeeded::DontCreate);
 		
 		float CurrentTimestamp = GWorld->TimeSeconds;
 		for (int32 Index = 0; Index < 2; ++Index)
@@ -165,7 +146,7 @@ bool FVisualLogSegmentsTest::RunTest(const FString& Parameters)
 		UE_VLOG_SEGMENT(GWorld, LogVisual, Log, StartPoint, EndPoint, FColor::Red, TEXT("Simple segment log"));
 		CHECK_SUCCESS(Context.Device.LastObject == NULL);
 		CHECK_SUCCESS(Context.Device.LastEntry.TimeStamp == -1);
-		FVisualLogEntry* CurrentEntry = FVisualLogger::Get().GetEntryToWrite(GWorld, GWorld->TimeSeconds, VisualLogger::DontCreate);
+		FVisualLogEntry* CurrentEntry = FVisualLogger::Get().GetEntryToWrite(GWorld, GWorld->TimeSeconds, ECreateIfNeeded::DontCreate);
 
 		float CurrentTimestamp = GWorld->TimeSeconds;
 		for (int32 Index = 0; Index < 2; ++Index)
@@ -173,7 +154,7 @@ bool FVisualLogSegmentsTest::RunTest(const FString& Parameters)
 			CHECK_SUCCESS(CurrentEntry != NULL);
 			CHECK_SUCCESS(CurrentEntry->TimeStamp == CurrentTimestamp);
 			CHECK_SUCCESS(CurrentEntry->ElementsToDraw.Num() == 1);
-			CHECK_SUCCESS(CurrentEntry->ElementsToDraw[0].GetType() == FVisualLogEntry::FElementToDraw::Segment);
+			CHECK_SUCCESS(CurrentEntry->ElementsToDraw[0].GetType() == EVisualLoggerShapeElement::Segment);
 			CHECK_SUCCESS(CurrentEntry->ElementsToDraw[0].Points.Num() == 2);
 			CHECK_SUCCESS(CurrentEntry->ElementsToDraw[0].Points[0] == StartPoint);
 			CHECK_SUCCESS(CurrentEntry->ElementsToDraw[0].Points[1] == EndPoint);
@@ -206,7 +187,7 @@ bool FVisualLogEventsTest::RunTest(const FString& Parameters)
 	CHECK_SUCCESS(EventTest3.Name == TEXT("EventTest3"));
 	CHECK_SUCCESS(EventTest3.FriendlyDesc == TEXT("Third simple event for tests"));
 
-	FVisualLogEntry* CurrentEntry = FVisualLogger::Get().GetEntryToWrite(GWorld, GWorld->TimeSeconds, VisualLogger::DontCreate);
+	FVisualLogEntry* CurrentEntry = FVisualLogger::Get().GetEntryToWrite(GWorld, GWorld->TimeSeconds, ECreateIfNeeded::DontCreate);
 	float CurrentTimestamp = GWorld->TimeSeconds;
 	UE_VLOG_EVENTS(GWorld, NAME_None, EventTest);
 	CHECK_SUCCESS(CurrentEntry != NULL);
diff --git a/Engine/Source/Runtime/Engine/Private/VisualLogger/VisualLoggerBinaryFileDevice.cpp b/Engine/Source/Runtime/Engine/Private/VisualLogger/VisualLoggerBinaryFileDevice.cpp
--- a/Engine/Source/Runtime/Engine/Private/VisualLogger/VisualLoggerBinaryFileDevice.cpp
+++ b/Engine/Source/Runtime/Engine/Private/VisualLogger/VisualLoggerBinaryFileDevice.cpp
@@ -3,12 +3,6 @@
 #include "VisualLogger/VisualLoggerBinaryFileDevice.h"
 
 #if ENABLE_VISUAL_LOG
-FArchive& operator<<(FArchive& Ar, FVisualLoggerBinaryFileDevice::FFrameCacheItem& FrameCacheItem)
-{
-	Ar << FrameCacheItem.OwnerName;
-	Ar << FrameCacheItem.Entry;
-	return Ar;
-}
 
 FVisualLoggerBinaryFileDevice::FVisualLoggerBinaryFileDevice()
 	: FileArchive(NULL)
@@ -38,7 +32,7 @@ void FVisualLoggerBinaryFileDevice::StartRecordingToFile(float TimeStamp)
 
 	StartRecordingTime = TimeStamp;
 	LastLogTimeStamp = StartRecordingTime;
-	TempFileName = FVisualLoggerHelpers::GenerateTemporaryFilename(TEXT("bvlog"));
+	TempFileName = FVisualLoggerHelpers::GenerateTemporaryFilename(VISLOG_FILENAME_EXT);
 	const FString FullFilename = FString::Printf(TEXT("%slogs/%s"), *FPaths::GameSavedDir(), *TempFileName);
 
 	FileArchive = IFileManager::Get().CreateFileWriter(*FullFilename);
@@ -83,22 +77,16 @@ void FVisualLoggerBinaryFileDevice::SetFileName(const FString& InFileName)
 	FileName = InFileName;
 }
 
-void FVisualLoggerBinaryFileDevice::Serialize(const class UObject* LogOwner, FName OwnerName, const FVisualLogEntry& LogEntry)
+void FVisualLoggerBinaryFileDevice::Serialize(const UObject* LogOwner, FName OwnerName, const FVisualLogEntry& LogEntry)
 {
 	const int32 NumEntries = FrameCache.Num();
-	if (NumEntries> 0 && FrameCache[NumEntries - 1].Entry.TimeStamp + FrameCacheLenght < LogEntry.TimeStamp && FileArchive)
+	if (NumEntries> 0 && LastLogTimeStamp + FrameCacheLenght <= LogEntry.TimeStamp && FileArchive)
 	{
-		int32 FrameTag = 0xFAFA;
-		*FileArchive << FrameTag;
-
-		int32 FrameCacheNum = FrameCache.Num();
-		*FileArchive << FrameCacheNum;
-
-		*FileArchive << FrameCache;
+		FVisualLoggerHelpers::Serialize(*FileArchive, FrameCache);
 		FrameCache.Reset();
 	}
 
 	LastLogTimeStamp = LogEntry.TimeStamp;
-	FrameCache.Add(FFrameCacheItem(OwnerName, LogEntry));
+	FrameCache.Add(FVisualLogEntryItem(OwnerName, LogEntry));
 }
 #endif
diff --git a/Engine/Source/Runtime/Engine/Private/VisualLogger/VisualLoggerKismetLibrary.cpp b/Engine/Source/Runtime/Engine/Private/VisualLogger/VisualLoggerKismetLibrary.cpp
--- a/Engine/Source/Runtime/Engine/Private/VisualLogger/VisualLoggerKismetLibrary.cpp
+++ b/Engine/Source/Runtime/Engine/Private/VisualLogger/VisualLoggerKismetLibrary.cpp
@@ -1,7 +1,7 @@
 // Copyright 1998-2014 Epic Games, Inc. All Rights Reserved.
 
 #include "EnginePrivate.h"
-#include "VisualLog.h"
+#include "VisualLogger/VisualLogger.h"
 #include "VisualLogger/VisualLogger.h"
 #include "VisualLogger/VisualLoggerKismetLibrary.h"
 
diff --git a/Engine/Source/Runtime/Engine/Private/VisualLogger/VisualLoggerTypes.cpp b/Engine/Source/Runtime/Engine/Private/VisualLogger/VisualLoggerTypes.cpp
--- a/Engine/Source/Runtime/Engine/Private/VisualLogger/VisualLoggerTypes.cpp
+++ b/Engine/Source/Runtime/Engine/Private/VisualLogger/VisualLoggerTypes.cpp
@@ -1,11 +1,28 @@
 // Copyright 1998-2014 Epic Games, Inc. All Rights Reserved.
 
 #include "EnginePrivate.h"
-#include "VisualLog.h"
 #include "VisualLogger/VisualLogger.h"
 
 #if ENABLE_VISUAL_LOG
 
+#define VISUAL_LOGGER_MAGIC_NUMBER 0xFAFAAFAF
+
+bool FVisualLogStatusCategory::GetDesc(int32 Index, FString& Key, FString& Value) const
+{
+	if (Data.IsValidIndex(Index))
+	{
+		int32 SplitIdx = INDEX_NONE;
+		if (Data[Index].FindChar(TEXT('|'), SplitIdx))
+		{
+			Key = Data[Index].Left(SplitIdx);
+			Value = Data[Index].Mid(SplitIdx + 1);
+			return true;
+		}
+	}
+
+	return false;
+}
+
 FVisualLogEntry::FVisualLogEntry(const FVisualLogEntry& Entry)
 {
 	TimeStamp = Entry.TimeStamp;
@@ -19,7 +36,7 @@ FVisualLogEntry::FVisualLogEntry(const FVisualLogEntry& Entry)
 	DataBlocks = Entry.DataBlocks;
 }
 
-FVisualLogEntry::FVisualLogEntry(const class AActor* InActor, TArray<TWeakObjectPtr<UObject> >* Children)
+FVisualLogEntry::FVisualLogEntry(const AActor* InActor, TArray<TWeakObjectPtr<UObject> >* Children)
 {
 	if (InActor && InActor->IsPendingKill() == false)
 	{
@@ -70,6 +87,18 @@ FVisualLogEntry::FVisualLogEntry(float InTimeStamp, FVector InLocation, const UO
 	}
 }
 
+void FVisualLogEntry::Reset()
+{
+	TimeStamp = -1;
+	Location = FVector::ZeroVector;
+	Events.Reset();
+	LogLines.Reset();
+	Status.Reset();
+	ElementsToDraw.Reset();
+	HistogramSamples.Reset();
+	DataBlocks.Reset();
+}
+
 int32 FVisualLogEntry::AddEvent(const FVisualLogEventBase& Event)
 {
 	return Events.Add(Event);
@@ -77,70 +106,81 @@ int32 FVisualLogEntry::AddEvent(const FVisualLogEventBase& Event)
 
 void FVisualLogEntry::AddText(const FString& TextLine, const FName& CategoryName)
 {
-	LogLines.Add(FLogLine(CategoryName, ELogVerbosity::All, TextLine));
+	LogLines.Add(FVisualLogLine(CategoryName, ELogVerbosity::All, TextLine));
 }
 
 void FVisualLogEntry::AddElement(const TArray<FVector>& Points, const FName& CategoryName, const FColor& Color, const FString& Description, uint16 Thickness)
 {
-	FElementToDraw Element(Description, Color, Thickness, CategoryName);
+	FVisualLogShapeElement Element(Description, Color, Thickness, CategoryName);
 	Element.Points = Points;
-	Element.Type = FElementToDraw::Path;
+	Element.Type = EVisualLoggerShapeElement::Path;
 	ElementsToDraw.Add(Element);
 }
 
 void FVisualLogEntry::AddElement(const FVector& Point, const FName& CategoryName, const FColor& Color, const FString& Description, uint16 Thickness)
 {
-	FElementToDraw Element(Description, Color, Thickness, CategoryName);
+	FVisualLogShapeElement Element(Description, Color, Thickness, CategoryName);
 	Element.Points.Add(Point);
-	Element.Type = FElementToDraw::SinglePoint;
+	Element.Type = EVisualLoggerShapeElement::SinglePoint;
 	ElementsToDraw.Add(Element);
 }
 
 void FVisualLogEntry::AddElement(const FVector& Start, const FVector& End, const FName& CategoryName, const FColor& Color, const FString& Description, uint16 Thickness)
 {
-	FElementToDraw Element(Description, Color, Thickness, CategoryName);
+	FVisualLogShapeElement Element(Description, Color, Thickness, CategoryName);
 	Element.Points.Reserve(2);
 	Element.Points.Add(Start);
 	Element.Points.Add(End);
-	Element.Type = FElementToDraw::Segment;
+	Element.Type = EVisualLoggerShapeElement::Segment;
 	ElementsToDraw.Add(Element);
 }
 
 void FVisualLogEntry::AddElement(const FBox& Box, const FName& CategoryName, const FColor& Color, const FString& Description, uint16 Thickness)
 {
-	FElementToDraw Element(Description, Color, Thickness, CategoryName);
+	FVisualLogShapeElement Element(Description, Color, Thickness, CategoryName);
 	Element.Points.Reserve(2);
 	Element.Points.Add(Box.Min);
 	Element.Points.Add(Box.Max);
-	Element.Type = FElementToDraw::Box;
+	Element.Type = EVisualLoggerShapeElement::Box;
 	ElementsToDraw.Add(Element);
 }
 
 void FVisualLogEntry::AddElement(const FVector& Orgin, const FVector& Direction, float Length, float AngleWidth, float AngleHeight, const FName& CategoryName, const FColor& Color, const FString& Description, uint16 Thickness)
 {
-	FElementToDraw Element(Description, Color, Thickness, CategoryName);
+	FVisualLogShapeElement Element(Description, Color, Thickness, CategoryName);
 	Element.Points.Reserve(3);
 	Element.Points.Add(Orgin);
 	Element.Points.Add(Direction);
 	Element.Points.Add(FVector(Length, AngleWidth, AngleHeight));
-	Element.Type = FElementToDraw::Cone;
+	Element.Type = EVisualLoggerShapeElement::Cone;
 	ElementsToDraw.Add(Element);
 }
 
 void FVisualLogEntry::AddElement(const FVector& Start, const FVector& End, float Radius, const FName& CategoryName, const FColor& Color, const FString& Description, uint16 Thickness)
 {
-	FElementToDraw Element(Description, Color, Thickness, CategoryName);
+	FVisualLogShapeElement Element(Description, Color, Thickness, CategoryName);
 	Element.Points.Reserve(3);
 	Element.Points.Add(Start);
 	Element.Points.Add(End);
 	Element.Points.Add(FVector(Radius, Thickness, 0));
-	Element.Type = FElementToDraw::Cylinder;
+	Element.Type = EVisualLoggerShapeElement::Cylinder;
+	ElementsToDraw.Add(Element);
+}
+
+void FVisualLogEntry::AddElement(const FVector& Center, float HalfHeight, float Radius, const FQuat & Rotation, const FName& CategoryName, const FColor& Color, const FString& Description)
+{
+	FVisualLogShapeElement Element(Description, Color, 0, CategoryName);
+	Element.Points.Reserve(3);
+	Element.Points.Add(Center);
+	Element.Points.Add(FVector(HalfHeight, Radius, Rotation.X));
+	Element.Points.Add(FVector(Rotation.Y, Rotation.Z, Rotation.W));
+	Element.Type = EVisualLoggerShapeElement::Capsule;
 	ElementsToDraw.Add(Element);
 }
 
 void FVisualLogEntry::AddHistogramData(const FVector2D& DataSample, const FName& CategoryName, const FName& GraphName, const FName& DataName)
 {
-	FHistogramSample Sample;
+	FVisualLogHistogramSample Sample;
 	Sample.Category = CategoryName;
 	Sample.GraphName = GraphName;
 	Sample.DataName = DataName;
@@ -149,9 +189,9 @@ void FVisualLogEntry::AddHistogramData(const FVector2D& DataSample, const FName&
 	HistogramSamples.Add(Sample);
 }
 
-FVisualLogEntry::FDataBlock& FVisualLogEntry::AddDataBlock(const FString& TagName, const TArray<uint8>& BlobDataArray, const FName& CategoryName)
+FVisualLogDataBlock& FVisualLogEntry::AddDataBlock(const FString& TagName, const TArray<uint8>& BlobDataArray, const FName& CategoryName)
 {
-	FDataBlock DataBlock;
+	FVisualLogDataBlock DataBlock;
 	DataBlock.Category = CategoryName;
 	DataBlock.TagName = *TagName;
 	DataBlock.Data = BlobDataArray;
@@ -160,15 +200,190 @@ FVisualLogEntry::FDataBlock& FVisualLogEntry::AddDataBlock(const FString& TagNam
 	return DataBlocks[Index];
 }
 
-void FVisualLogEntry::AddCapsule(FVector const& Center, float HalfHeight, float Radius, const FQuat & Rotation, const FName& CategoryName, const FColor& Color, const FString& Description)
+FArchive& operator<<(FArchive& Ar, FVisualLogDataBlock& Data)
 {
-	FElementToDraw Element(Description, Color, 0, CategoryName);
-	Element.Points.Reserve(3);
-	Element.Points.Add(Center);
-	Element.Points.Add(FVector(HalfHeight, Radius, Rotation.X));
-	Element.Points.Add(FVector(Rotation.Y, Rotation.Z, Rotation.W));
-	Element.Type = FElementToDraw::Capsule;
-	ElementsToDraw.Add(Element);
+	FVisualLoggerHelpers::Serialize(Ar, Data.TagName);
+	FVisualLoggerHelpers::Serialize(Ar, Data.Category);
+	Ar << Data.Verbosity;
+	Ar << Data.Data;
+	Ar << Data.UniqueId;
+
+	return Ar;
+}
+
+FArchive& operator<<(FArchive& Ar, FVisualLogHistogramSample& Sample)
+{
+	FVisualLoggerHelpers::Serialize(Ar, Sample.Category);
+	FVisualLoggerHelpers::Serialize(Ar, Sample.GraphName);
+	FVisualLoggerHelpers::Serialize(Ar, Sample.DataName);
+	Ar << Sample.Verbosity;
+	Ar << Sample.SampleValue;
+	Ar << Sample.UniqueId;
+
+	return Ar;
+}
+
+FArchive& operator<<(FArchive& Ar, EVisualLoggerShapeElement& Shape)
+{
+	uint8 ShapeAsInt = (uint8)Shape;
+	Ar << ShapeAsInt;
+
+	if (Ar.IsLoading())
+	{
+		Shape = (EVisualLoggerShapeElement)ShapeAsInt;
+	}
+
+	return Ar;
+}
+
+FArchive& operator<<(FArchive& Ar, FVisualLogShapeElement& Element)
+{
+	FVisualLoggerHelpers::Serialize(Ar, Element.Category);
+	Ar << Element.Description;
+	Ar << Element.Verbosity;
+	Ar << Element.Points;
+	Ar << Element.UniqueId;
+	Ar << Element.Type;
+	Ar << Element.Color;
+	Ar << Element.Thicknes;
+
+	return Ar;
+}
+
+FArchive& operator<<(FArchive& Ar, FVisualLogEvent& Event)
+{
+	Ar << Event.Name;
+	Ar << Event.UserFriendlyDesc;
+	Ar << Event.Verbosity;
+
+	int32 NumberOfTags = Event.EventTags.Num();
+	Ar << NumberOfTags;
+	if (Ar.IsLoading())
+	{
+		for (int32 Index = 0; Index < NumberOfTags; ++Index)
+		{
+			FName Key = NAME_None;
+			int32 Value = 0;
+			FVisualLoggerHelpers::Serialize(Ar, Key);
+			Ar << Value;
+			Event.EventTags.Add(Key, Value);
+		}
+	}
+	else
+	{
+		for (auto& CurrentTag : Event.EventTags)
+		{
+			FVisualLoggerHelpers::Serialize(Ar, CurrentTag.Key);
+			Ar << CurrentTag.Value;
+		}
+	}
+
+	Ar << Event.Counter;
+	Ar << Event.UserData;
+	FVisualLoggerHelpers::Serialize(Ar, Event.TagName);
+
+	return Ar;
+}
+
+FArchive& operator<<(FArchive& Ar, FVisualLogLine& LogLine)
+{
+	FVisualLoggerHelpers::Serialize(Ar, LogLine.Category);
+	FVisualLoggerHelpers::Serialize(Ar, LogLine.TagName);
+	Ar << LogLine.Verbosity;
+	Ar << LogLine.UniqueId;
+	Ar << LogLine.UserData;
+	Ar << LogLine.Line;
+	return Ar;
+}
+
+FArchive& operator<<(FArchive& Ar, FVisualLogStatusCategory& Status)
+{
+	Ar << Status.Category;
+	Ar << Status.Data;
+
+	return Ar;
+}
+
+FArchive& operator<<(FArchive& Ar, FVisualLogEntry& LogEntry)
+{
+	Ar << LogEntry.TimeStamp;
+	Ar << LogEntry.Location;
+	Ar << LogEntry.LogLines;
+	Ar << LogEntry.Status;
+	Ar << LogEntry.Events;
+	Ar << LogEntry.ElementsToDraw;
+	Ar << LogEntry.DataBlocks;
+
+	return Ar;
+}
+
+FArchive& operator<<(FArchive& Ar, FVisualLogDevice::FVisualLogEntryItem& FrameCacheItem)
+{
+	FVisualLoggerHelpers::Serialize(Ar, FrameCacheItem.OwnerName);
+	Ar << FrameCacheItem.Entry;
+	return Ar;
+}
+
+FString FVisualLoggerHelpers::GenerateTemporaryFilename(const FString& FileExt)
+{
+	return FString::Printf(TEXT("VTEMP_%s.%s"), *FDateTime::Now().ToString(), *FileExt);
+}
+
+FString FVisualLoggerHelpers::GenerateFilename(const FString& TempFileName, const FString& Prefix, float StartRecordingTime, float EndTimeStamp)
+{
+	const FString TempFullFilename = FString::Printf(TEXT("%slogs/%s"), *FPaths::GameSavedDir(), *TempFileName);
+	const FString FullFilename = FString::Printf(TEXT("%slogs/%s_%s"), *FPaths::GameSavedDir(), *Prefix, *TempFileName);
+	const FString TimeFrameString = FString::Printf(TEXT("%d-%d_"), FMath::TruncToInt(StartRecordingTime), FMath::TruncToInt(EndTimeStamp));
+	return FullFilename.Replace(TEXT("VTEMP_"), *TimeFrameString, ESearchCase::CaseSensitive);
+}
+
+FArchive& FVisualLoggerHelpers::Serialize(FArchive& Ar, FName& Name)
+{
+	// Serialize the FName as a string
+	if (Ar.IsLoading())
+	{
+		FString StringName;
+		Ar << StringName;
+		Name = FName(*StringName);
+	}
+	else
+	{
+		FString StringName = Name.ToString();
+		Ar << StringName;
+	}
+	return Ar;
+}
+
+FArchive& FVisualLoggerHelpers::Serialize(FArchive& Ar, TArray<FVisualLogDevice::FVisualLogEntryItem>& RecordedLogs)
+{
+	Ar.UsingCustomVersion(EVisualLoggerVersion::GUID);
+
+	if (Ar.IsLoading())
+	{
+		TArray<FVisualLogDevice::FVisualLogEntryItem> CurrentFrame;;
+
+		while (Ar.AtEnd() == false)
+		{
+			int32 FrameTag = VISUAL_LOGGER_MAGIC_NUMBER;
+			Ar << FrameTag;
+			if (FrameTag != VISUAL_LOGGER_MAGIC_NUMBER)
+			{
+				break;
+			}
+
+			Ar << CurrentFrame;
+			RecordedLogs.Append(CurrentFrame);
+			CurrentFrame.Reset();
+		}
+	}
+	else
+	{
+		int32 FrameTag = VISUAL_LOGGER_MAGIC_NUMBER;
+		Ar << FrameTag;
+		Ar << RecordedLogs;
+	}
+
+	return Ar;
 }
 
 #endif //ENABLE_VISUAL_LOG
diff --git a/Engine/Source/Runtime/Engine/Private/World.cpp b/Engine/Source/Runtime/Engine/Private/World.cpp
--- a/Engine/Source/Runtime/Engine/Private/World.cpp
+++ b/Engine/Source/Runtime/Engine/Private/World.cpp
@@ -30,7 +30,7 @@
 #include "TickTaskManagerInterface.h"
 #include "FXSystem.h"
 #include "SoundDefinitions.h"
-#include "VisualLog.h"
+#include "VisualLogger/VisualLogger.h"
 #include "LevelUtils.h"
 #include "PhysicsPublic.h"
 #include "AI/AISystemBase.h"
diff --git a/Engine/Source/Runtime/Engine/Public/VisualLog.h b/Engine/Source/Runtime/Engine/Public/VisualLog.h
deleted file mode 100644
diff --git a/Engine/Source/Runtime/Engine/Public/VisualLogger/VisualLogger.h b/Engine/Source/Runtime/Engine/Public/VisualLogger/VisualLogger.h
--- a/Engine/Source/Runtime/Engine/Public/VisualLogger/VisualLogger.h
+++ b/Engine/Source/Runtime/Engine/Public/VisualLogger/VisualLogger.h
@@ -14,6 +14,9 @@
 // Text, regular log
 #define UE_VLOG(LogOwner, CategoryName, Verbosity, Format, ...) FVisualLogger::CategorizedLogf(LogOwner, CategoryName, ELogVerbosity::Verbosity, INDEX_NONE, Format, ##__VA_ARGS__)
 #define UE_CVLOG(Condition, LogOwner, CategoryName, Verbosity, Format, ...)  if(Condition) {UE_VLOG(LogOwner, CategoryName, Verbosity, Format, ##__VA_ARGS__);} 
+// Text, log with output to regular unreal logs too
+#define UE_VLOG_UELOG(LogOwner, CategoryName, Verbosity, Format, ...) { FVisualLogger::CategorizedLogf(LogOwner, CategoryName, ELogVerbosity::Verbosity, INDEX_NONE, Format, ##__VA_ARGS__); UE_LOG(CategoryName, Verbosity, Format, ##__VA_ARGS__); }
+#define UE_CVLOG_UELOG(Condition, LogOwner, CategoryName, Verbosity, Format, ...)  if(Condition) {UE_VLOG_DEBUG(LogOwner, CategoryName, Verbosity, Format, ##__VA_ARGS__);} 
 // Segment shape
 #define UE_VLOG_SEGMENT(LogOwner, CategoryName, Verbosity, SegmentStart, SegmentEnd, Color, Format, ...) FVisualLogger::GeometryShapeLogf(LogOwner, CategoryName, ELogVerbosity::Verbosity, INDEX_NONE, SegmentStart, SegmentEnd, Color, Format, ##__VA_ARGS__)
 #define UE_CVLOG_SEGMENT(Condition, LogOwner, CategoryName, Verbosity, SegmentStart, SegmentEnd, Color, Format, ...) if(Condition) {UE_VLOG_SEGMENT(LogOwner, CategoryName, Verbosity, INDEX_NONE, SegmentStart, SegmentEnd, Color, Format, ##__VA_ARGS__);}
@@ -55,6 +58,8 @@
 
 #define UE_VLOG(Actor, CategoryName, Verbosity, Format, ...)
 #define UE_CVLOG(Condition, Actor, CategoryName, Verbosity, Format, ...)
+#define UE_VLOG_UELOG(LogOwner, CategoryName, Verbosity, Format, ...)
+#define UE_CVLOG_UELOG(Condition, Actor, CategoryName, Verbosity, Format, ...)
 #define UE_VLOG_SEGMENT(Actor, CategoryName, Verbosity, SegmentStart, SegmentEnd, Color, DescriptionFormat, ...)
 #define UE_CVLOG_SEGMENT(Condition, Actor, CategoryName, Verbosity, SegmentStart, SegmentEnd, Color, DescriptionFormat, ...)
 #define UE_VLOG_SEGMENT_THICK(Actor, CategoryName, Verbosity, SegmentStart, SegmentEnd, Color, DescriptionFormat, ...)
@@ -136,7 +141,7 @@ public:
 		return GVisLog;
 	}
 
-	virtual ~FVisualLogger();
+	virtual ~FVisualLogger() {}
 
 	// called on engine shutdown to flush all, etc.
 	virtual void Shutdown();
@@ -188,7 +193,7 @@ public:
 	/** Remove visual logger output device */
 	TArray<FVisualLogDevice*>& GetDevices() { return OutputDevices; }
 	/** Returns  current entry for given TimeStap or creates another one  but first it serialize previous entry as completed to vislog devices. Use VisualLogger::DontCreate to get current entry without serialization*/
-	FVisualLogEntry* GetEntryToWrite(const class UObject* Object, float TimeStamp, VisualLogger::ECreateIfNeeded ShouldCreate = VisualLogger::Create);
+	FVisualLogEntry* GetEntryToWrite(const class UObject* Object, float TimeStamp, ECreateIfNeeded ShouldCreate = ECreateIfNeeded::Create);
 	/** flush and serialize data if timestamp allows it */
 	virtual void Flush() override;
 
@@ -241,31 +246,6 @@ protected:
 	bool UseBinaryFileDevice;
 };
 
-/**
- * Interface for Visual Logger Device
- */
-class FVisualLogDevice
-{
-public:
-	virtual void Cleanup(bool bReleaseMemory = false) = 0;
-	virtual void StartRecordingToFile(float TImeStamp) = 0;
-	virtual void StopRecordingToFile(float TImeStamp) = 0;
-	virtual void Serialize(const class UObject* LogOwner, FName OwnerName, const FVisualLogEntry& LogEntry) = 0;
-	virtual void SetFileName(const FString& InFileName) = 0;
-	virtual bool HasFlags(int32 InFlags) { return !!(InFlags & (VisualLogger::CanSaveToFile | VisualLogger::StoreLogsLocally)); }
-};
-
-/**
- * Interface for extensions
- */
-class FVisualLogExtensionInterface
-{
-public:
-	virtual void OnTimestampChange(float Timestamp, class UWorld* InWorld, class AActor* HelperActor) = 0;
-	virtual void DrawData(class UWorld* InWorld, class UCanvas* Canvas, class AActor* HelperActor, const FName& TagName, const FVisualLogEntry::FDataBlock& DataBlock, float Timestamp) = 0;
-	virtual void DisableDrawingForData(class UWorld* InWorld, class UCanvas* Canvas, class AActor* HelperActor, const FName& TagName, const FVisualLogEntry::FDataBlock& DataBlock, float Timestamp) = 0;
-	virtual void LogEntryLineSelectionChanged(TSharedPtr<struct FLogEntryItem> SelectedItem, int64 UserData, FName TagName) = 0;
-};
 
 #include "VisualLogger.inl"
 
diff --git a/Engine/Source/Runtime/Engine/Public/VisualLogger/VisualLogger.inl b/Engine/Source/Runtime/Engine/Public/VisualLogger/VisualLogger.inl
--- a/Engine/Source/Runtime/Engine/Public/VisualLogger/VisualLogger.inl
+++ b/Engine/Source/Runtime/Engine/Public/VisualLogger/VisualLogger.inl
@@ -32,7 +32,7 @@ bool CheckVisualLogInputInternal(const class UObject* Object, const struct FLogC
 }
 
 FORCEINLINE
-FVisualLogEntry* FVisualLogger::GetEntryToWrite(const class UObject* Object, float TimeStamp, VisualLogger::ECreateIfNeeded ShouldCreate)
+FVisualLogEntry* FVisualLogger::GetEntryToWrite(const class UObject* Object, float TimeStamp, ECreateIfNeeded ShouldCreate)
 {
 	FVisualLogEntry* CurrentEntry = NULL;
 	UObject * LogOwner = FVisualLogger::FindRedirection(Object);
@@ -41,7 +41,7 @@ FVisualLogEntry* FVisualLogger::GetEntryToWrite(const class UObject* Object, flo
 	if (CurrentEntryPerObject.Contains(LogOwner))
 	{
 		CurrentEntry = &CurrentEntryPerObject[LogOwner];
-		if (TimeStamp > CurrentEntry->TimeStamp && ShouldCreate == VisualLogger::Create)
+		if (TimeStamp > CurrentEntry->TimeStamp && ShouldCreate == ECreateIfNeeded::Create)
 		{
 			if (CurrentEntry->TimeStamp >= 0) //-1 means not initialized entry information
 			{
@@ -122,6 +122,7 @@ void FVisualLogger::Flush()
 FORCEINLINE
 VARARG_BODY(void, FVisualLogger::CategorizedLogf, const TCHAR*, VARARG_EXTRA(const class UObject* Object) VARARG_EXTRA(const struct FLogCategoryBase& Category) VARARG_EXTRA(ELogVerbosity::Type Verbosity) VARARG_EXTRA(int32 UniqueLogId))
 {
+	SCOPE_CYCLE_COUNTER(STAT_VisualLog);
 	UWorld *World = NULL;
 	FVisualLogEntry *CurrentEntry = NULL;
 	if (CheckVisualLogInputInternal(Object, Category, Verbosity, &World, &CurrentEntry) == false)
@@ -137,6 +138,7 @@ VARARG_BODY(void, FVisualLogger::CategorizedLogf, const TCHAR*, VARARG_EXTRA(con
 FORCEINLINE
 VARARG_BODY(void, FVisualLogger::GeometryShapeLogf, const TCHAR*, VARARG_EXTRA(const class UObject* Object) VARARG_EXTRA(const struct FLogCategoryBase& Category) VARARG_EXTRA(ELogVerbosity::Type Verbosity) VARARG_EXTRA(int32 UniqueLogId) VARARG_EXTRA(const FVector& Start) VARARG_EXTRA(const FVector& End) VARARG_EXTRA(const FColor& Color))
 {
+	SCOPE_CYCLE_COUNTER(STAT_VisualLog);
 	UWorld *World = NULL;
 	FVisualLogEntry *CurrentEntry = NULL;
 	if (CheckVisualLogInputInternal(Object, Category, Verbosity, &World, &CurrentEntry) == false)
@@ -152,6 +154,7 @@ VARARG_BODY(void, FVisualLogger::GeometryShapeLogf, const TCHAR*, VARARG_EXTRA(c
 FORCEINLINE
 VARARG_BODY(void, FVisualLogger::GeometryShapeLogf, const TCHAR*, VARARG_EXTRA(const class UObject* Object) VARARG_EXTRA(const struct FLogCategoryBase& Category) VARARG_EXTRA(ELogVerbosity::Type Verbosity) VARARG_EXTRA(int32 UniqueLogId) VARARG_EXTRA(const FVector& Location) VARARG_EXTRA(float Radius) VARARG_EXTRA(const FColor& Color))
 {
+	SCOPE_CYCLE_COUNTER(STAT_VisualLog);
 	UWorld *World = NULL;
 	FVisualLogEntry *CurrentEntry = NULL;
 	if (CheckVisualLogInputInternal(Object, Category, Verbosity, &World, &CurrentEntry) == false)
@@ -167,6 +170,7 @@ VARARG_BODY(void, FVisualLogger::GeometryShapeLogf, const TCHAR*, VARARG_EXTRA(c
 FORCEINLINE
 VARARG_BODY(void, FVisualLogger::GeometryShapeLogf, const TCHAR*, VARARG_EXTRA(const class UObject* Object) VARARG_EXTRA(const struct FLogCategoryBase& Category) VARARG_EXTRA(ELogVerbosity::Type Verbosity) VARARG_EXTRA(int32 UniqueLogId) VARARG_EXTRA(const FBox& Box) VARARG_EXTRA(const FColor& Color))
 {
+	SCOPE_CYCLE_COUNTER(STAT_VisualLog);
 	UWorld *World = NULL;
 	FVisualLogEntry *CurrentEntry = NULL;
 	if (CheckVisualLogInputInternal(Object, Category, Verbosity, &World, &CurrentEntry) == false)
@@ -182,6 +186,7 @@ VARARG_BODY(void, FVisualLogger::GeometryShapeLogf, const TCHAR*, VARARG_EXTRA(c
 FORCEINLINE
 VARARG_BODY(void, FVisualLogger::GeometryShapeLogf, const TCHAR*, VARARG_EXTRA(const class UObject* Object) VARARG_EXTRA(const struct FLogCategoryBase& Category) VARARG_EXTRA(ELogVerbosity::Type Verbosity) VARARG_EXTRA(int32 UniqueLogId) VARARG_EXTRA(const FVector& Orgin) VARARG_EXTRA(const FVector& Direction) VARARG_EXTRA(const float Length) VARARG_EXTRA(const float Angle)  VARARG_EXTRA(const FColor& Color))
 {
+	SCOPE_CYCLE_COUNTER(STAT_VisualLog);
 	UWorld *World = NULL;
 	FVisualLogEntry *CurrentEntry = NULL;
 	if (CheckVisualLogInputInternal(Object, Category, Verbosity, &World, &CurrentEntry) == false)
@@ -197,6 +202,7 @@ VARARG_BODY(void, FVisualLogger::GeometryShapeLogf, const TCHAR*, VARARG_EXTRA(c
 FORCEINLINE
 VARARG_BODY(void, FVisualLogger::GeometryShapeLogf, const TCHAR*, VARARG_EXTRA(const class UObject* Object) VARARG_EXTRA(const struct FLogCategoryBase& Category) VARARG_EXTRA(ELogVerbosity::Type Verbosity) VARARG_EXTRA(int32 UniqueLogId) VARARG_EXTRA(const FVector& Start) VARARG_EXTRA(const FVector& End) VARARG_EXTRA(const float Radius) VARARG_EXTRA(const FColor& Color))
 {
+	SCOPE_CYCLE_COUNTER(STAT_VisualLog);
 	UWorld *World = NULL;
 	FVisualLogEntry *CurrentEntry = NULL;
 	if (CheckVisualLogInputInternal(Object, Category, Verbosity, &World, &CurrentEntry) == false)
@@ -212,6 +218,7 @@ VARARG_BODY(void, FVisualLogger::GeometryShapeLogf, const TCHAR*, VARARG_EXTRA(c
 FORCEINLINE
 VARARG_BODY(void, FVisualLogger::GeometryShapeLogf, const TCHAR*, VARARG_EXTRA(const class UObject* Object) VARARG_EXTRA(const struct FLogCategoryBase& Category) VARARG_EXTRA(ELogVerbosity::Type Verbosity) VARARG_EXTRA(int32 UniqueLogId) VARARG_EXTRA(const FVector& Center) VARARG_EXTRA(float HalfHeight) VARARG_EXTRA(float Radius) VARARG_EXTRA(const FQuat& Rotation) VARARG_EXTRA(const FColor& Color))
 {
+	SCOPE_CYCLE_COUNTER(STAT_VisualLog);
 	UWorld *World = NULL;
 	FVisualLogEntry *CurrentEntry = NULL;
 	if (CheckVisualLogInputInternal(Object, Category, Verbosity, &World, &CurrentEntry) == false)
@@ -220,13 +227,14 @@ VARARG_BODY(void, FVisualLogger::GeometryShapeLogf, const TCHAR*, VARARG_EXTRA(c
 	}
 
 	GROWABLE_LOGF(
-		CurrentEntry->AddCapsule(Center, HalfHeight, Radius, Rotation, Category.GetCategoryName(), Color, Buffer);
+		CurrentEntry->AddElement(Center, HalfHeight, Radius, Rotation, Category.GetCategoryName(), Color, Buffer);
 	);
 }
 
 FORCEINLINE
 VARARG_BODY(void, FVisualLogger::HistogramDataLogf, const TCHAR*, VARARG_EXTRA(const class UObject* Object) VARARG_EXTRA(const struct FLogCategoryBase& Category) VARARG_EXTRA(ELogVerbosity::Type Verbosity) VARARG_EXTRA(int32 UniqueLogId) VARARG_EXTRA(FName GraphName) VARARG_EXTRA(FName DataName) VARARG_EXTRA(const FVector2D& Data) VARARG_EXTRA(const FColor& Color))
 {
+	SCOPE_CYCLE_COUNTER(STAT_VisualLog);
 	UWorld *World = NULL;
 	FVisualLogEntry *CurrentEntry = NULL;
 	if (CheckVisualLogInputInternal(Object, Category, Verbosity, &World, &CurrentEntry) == false)
@@ -277,6 +285,7 @@ void FVisualLogger::EventLog(const class UObject* Object, const FName EventTag1,
 FORCEINLINE
 void FVisualLogger::EventLog(const class UObject* Object, const FName EventTag1, const FVisualLogEventBase& Event, const FName EventTag2, const FName EventTag3, const FName EventTag4, const FName EventTag5, const FName EventTag6)
 {
+	SCOPE_CYCLE_COUNTER(STAT_VisualLog);
 	UWorld *World = NULL;
 	FVisualLogEntry *CurrentEntry = NULL;
 	const FLogCategory<ELogVerbosity::Log, ELogVerbosity::Log> Category(*Event.Name);
@@ -285,7 +294,7 @@ void FVisualLogger::EventLog(const class UObject* Object, const FName EventTag1,
 		return;
 	}
 
-	int32 Index = CurrentEntry->Events.Find(FVisualLogEntry::FLogEvent(Event));
+	int32 Index = CurrentEntry->Events.Find(FVisualLogEvent(Event));
 	if (Index != INDEX_NONE)
 	{
 		CurrentEntry->Events[Index].Counter++;
@@ -303,5 +312,3 @@ void FVisualLogger::EventLog(const class UObject* Object, const FName EventTag1,
 	CurrentEntry->Events[Index].EventTags.FindOrAdd(EventTag6)++;
 	CurrentEntry->Events[Index].EventTags.Remove(NAME_None);
 }
-
-#undef  CHECK_VLOG_DATA
diff --git a/Engine/Source/Runtime/Engine/Public/VisualLogger/VisualLoggerTypes.h b/Engine/Source/Runtime/Engine/Public/VisualLogger/VisualLoggerTypes.h
--- a/Engine/Source/Runtime/Engine/Public/VisualLogger/VisualLoggerTypes.h
+++ b/Engine/Source/Runtime/Engine/Public/VisualLogger/VisualLoggerTypes.h
@@ -2,22 +2,56 @@
 #pragma once
 #include "EngineDefines.h"
 
-namespace VisualLogger
+enum class ECreateIfNeeded : int8
 {
-	enum ECreateIfNeeded
-	{
-		Invalid = -1,
-		DontCreate = 0,
-		Create = 1,
-	};
+	Invalid = -1,
+	DontCreate = 0,
+	Create = 1,
+};
 
-	enum DeviceFlags
+// flags describing VisualLogger device's features
+namespace EVisualLoggerDeviceFlags
+{
+	enum Type
 	{
+		NoFlags = 0,
 		CanSaveToFile = 1,
 		StoreLogsLocally = 2,
 	};
 }
 
+// version for vlog binary file format
+namespace EVisualLoggerVersion
+{
+	enum Type
+	{
+		Initial = 0,
+		// -----<new versions can be added before this line>-------------------------------------------------
+		// - this needs to be the last line (see note below)
+		VersionPlusOne,
+		LatestVersion = VersionPlusOne - 1
+	};
+
+	// The GUID for this custom version number
+	extern ENGINE_API const FGuid GUID;
+}
+
+//types of shape elements
+enum class EVisualLoggerShapeElement : uint8
+{
+	Invalid = 0,
+	SinglePoint, // individual points. 
+	Segment, // pairs of points 
+	Path,	// sequence of point
+	Box,
+	Cone,
+	Cylinder,
+	Capsule,
+	// note that in order to remain backward compatibility in terms of log
+	// serialization new enum values need to be added at the end
+};
+
+#if ENABLE_VISUAL_LOG
 struct ENGINE_API FVisualLogEventBase
 {
 	const FString Name;
@@ -30,357 +64,300 @@ struct ENGINE_API FVisualLogEventBase
 	}
 };
 
-struct ENGINE_API FVisualLogEntry
+struct ENGINE_API FVisualLogEvent
 {
-	struct FLogEvent
-	{
-		FString Name;
-		FString UserFriendlyDesc;
-		TEnumAsByte<ELogVerbosity::Type> Verbosity;
-		TMap<FName, int32>	 EventTags;
-		int32 Counter;
-		int64 UserData;
-		FName TagName;
-
-		FLogEvent() : Counter(1) {}
-		FLogEvent(const FVisualLogEventBase& Event) : Counter(1)
-		{
-			Name = Event.Name;
-			UserFriendlyDesc = Event.FriendlyDesc;
-			Verbosity = Event.Verbosity;
-		}
+	FString Name;
+	FString UserFriendlyDesc;
+	TEnumAsByte<ELogVerbosity::Type> Verbosity;
+	TMap<FName, int32>	 EventTags;
+	int32 Counter;
+	int64 UserData;
+	FName TagName;
+
+	FVisualLogEvent() : Counter(1) { /* Empty */ }
+	FVisualLogEvent(const FVisualLogEventBase& Event);
+	FVisualLogEvent& operator=(const FVisualLogEventBase& Event);
+};
 
-		FLogEvent& operator = (const FVisualLogEventBase& Event)
-		{
-			Name = Event.Name;
-			UserFriendlyDesc = Event.FriendlyDesc;
-			Verbosity = Event.Verbosity;
-			return *this;
-		}
-	};
-	struct FLogLine
-	{
-		FString Line;
-		FName Category;
-		TEnumAsByte<ELogVerbosity::Type> Verbosity;
-		int32 UniqueId;
-		int64 UserData;
-		FName TagName;
-
-		FLogLine() {}
-
-		FLogLine(const FName& InCategory, ELogVerbosity::Type InVerbosity, const FString& InLine)
-			: Line(InLine), Category(InCategory), Verbosity(InVerbosity), UserData(0)
-		{}
-
-		FLogLine(const FName& InCategory, ELogVerbosity::Type InVerbosity, const FString& InLine, int64 InUserData)
-			: Line(InLine), Category(InCategory), Verbosity(InVerbosity), UserData(InUserData)
-		{}
-	};
+struct ENGINE_API FVisualLogLine
+{
+	FString Line;
+	FName Category;
+	TEnumAsByte<ELogVerbosity::Type> Verbosity;
+	int32 UniqueId;
+	int64 UserData;
+	FName TagName;
+
+	FVisualLogLine() { /* Empty */ }
+	FVisualLogLine(const FName& InCategory, ELogVerbosity::Type InVerbosity, const FString& InLine);
+	FVisualLogLine(const FName& InCategory, ELogVerbosity::Type InVerbosity, const FString& InLine, int64 InUserData);
+};
 
-	struct FStatusCategory
-	{
-		TArray<FString> Data;
-		FString Category;
-		int32 UniqueId;
+struct ENGINE_API FVisualLogStatusCategory
+{
+	TArray<FString> Data;
+	FString Category;
+	int32 UniqueId;
 
-		FORCEINLINE void Add(const FString& Key, const FString& Value)
-		{
-			Data.Add(FString(Key).AppendChar(TEXT('|')) + Value);
-		}
+	void Add(const FString& Key, const FString& Value);
+	bool GetDesc(int32 Index, FString& Key, FString& Value) const;
+};
 
-		FORCEINLINE bool GetDesc(int32 Index, FString& Key, FString& Value) const
-		{
-			if (Data.IsValidIndex(Index))
-			{
-				int32 SplitIdx = INDEX_NONE;
-				if (Data[Index].FindChar(TEXT('|'), SplitIdx))
-				{
-					Key = Data[Index].Left(SplitIdx);
-					Value = Data[Index].Mid(SplitIdx + 1);
-					return true;
-				}
-			}
-
-			return false;
-		}
+struct ENGINE_API FVisualLogShapeElement
+{
+	FString Description;
+	FName Category;
+	TEnumAsByte<ELogVerbosity::Type> Verbosity;
+	TArray<FVector> Points;
+	int32 UniqueId;
+	EVisualLoggerShapeElement Type;
+	uint8 Color;
+	union
+	{
+		uint16 Thicknes;
+		uint16 Radius;
 	};
 
-	struct FElementToDraw
-	{
-		friend struct FVisualLogEntry;
-
-		enum EElementType {
-			Invalid = 0,
-			SinglePoint, // individual points. 
-			Segment, // pairs of points 
-			Path,	// sequence of point
-			Box,
-			Cone,
-			Cylinder,
-			Capsule,
-			// note that in order to remain backward compatibility in terms of log
-			// serialization new enum values need to be added at the end
-		};
-
-		FString Description;
-		FName Category;
-		TEnumAsByte<ELogVerbosity::Type> Verbosity;
-		TArray<FVector> Points;
-		int32 UniqueId;
-
-	public:
-		uint8 Type;
-		uint8 Color;
-		union
-		{
-			uint16 Thicknes;
-			uint16 Radius;
-		};
+	FVisualLogShapeElement(EVisualLoggerShapeElement InType = EVisualLoggerShapeElement::Invalid);
+	FVisualLogShapeElement(const FString& InDescription, const FColor& InColor, uint16 InThickness, const FName& InCategory);
+	void SetColor(const FColor& InColor);
+	EVisualLoggerShapeElement GetType() const;
+	void SetType(EVisualLoggerShapeElement InType);
+	FColor GetFColor() const;
+};
 
-		FElementToDraw(EElementType InType = Invalid) : Verbosity(ELogVerbosity::All), Type(InType), Color(0xff), Thicknes(0)
-		{}
+struct ENGINE_API FVisualLogHistogramSample
+{
+	FName Category;
+	TEnumAsByte<ELogVerbosity::Type> Verbosity;
+	FName GraphName;
+	FName DataName;
+	FVector2D SampleValue;
+	int32 UniqueId;
+};
 
-		FElementToDraw(const FString& InDescription, const FColor& InColor, uint16 InThickness, const FName& InCategory) : Category(InCategory), Verbosity(ELogVerbosity::All), Type(uint16(Invalid)), Thicknes(InThickness)
-		{
-			if (InDescription.IsEmpty() == false)
-			{
-				Description = InDescription;
-			}
-			SetColor(InColor);
-		}
+struct ENGINE_API FVisualLogDataBlock
+{
+	FName TagName;
+	FName Category;
+	TEnumAsByte<ELogVerbosity::Type> Verbosity;
+	TArray<uint8>	Data;
+	int32 UniqueId;
+};
+#endif  //ENABLE_VISUAL_LOG
 
-		FORCEINLINE void SetColor(const FColor& InColor)
-		{
-			Color = ((InColor.DWColor() >> 30) << 6)
-				| (((InColor.DWColor() & 0x00ff0000) >> 22) << 4)
-				| (((InColor.DWColor() & 0x0000ff00) >> 14) << 2)
-				| ((InColor.DWColor() & 0x000000ff) >> 6);
-		}
+struct ENGINE_API FVisualLogEntry
+{
+#if ENABLE_VISUAL_LOG
+	float TimeStamp;
+	FVector Location;
+
+	TArray<FVisualLogEvent> Events;
+	TArray<FVisualLogLine> LogLines;
+	TArray<FVisualLogStatusCategory> Status;
+	TArray<FVisualLogShapeElement> ElementsToDraw;
+	TArray<FVisualLogHistogramSample>	HistogramSamples;
+	TArray<FVisualLogDataBlock>	DataBlocks;
+
+	FVisualLogEntry() { Reset(); }
+	FVisualLogEntry(const FVisualLogEntry& Entry);
+	FVisualLogEntry(const class AActor* InActor, TArray<TWeakObjectPtr<UObject> >* Children);
+	FVisualLogEntry(float InTimeStamp, FVector InLocation, const UObject* Object, TArray<TWeakObjectPtr<UObject> >* Children);
+
+	void Reset();
+
+	void AddText(const FString& TextLine, const FName& CategoryName);
+	// path
+	void AddElement(const TArray<FVector>& Points, const FName& CategoryName, const FColor& Color = FColor::White, const FString& Description = TEXT(""), uint16 Thickness = 0);
+	// location
+	void AddElement(const FVector& Point, const FName& CategoryName, const FColor& Color = FColor::White, const FString& Description = TEXT(""), uint16 Thickness = 0);
+	// segment
+	void AddElement(const FVector& Start, const FVector& End, const FName& CategoryName, const FColor& Color = FColor::White, const FString& Description = TEXT(""), uint16 Thickness = 0);
+	// box
+	void AddElement(const FBox& Box, const FName& CategoryName, const FColor& Color = FColor::White, const FString& Description = TEXT(""), uint16 Thickness = 0);
+	// Cone
+	void AddElement(const FVector& Orgin, const FVector& Direction, float Length, float AngleWidth, float AngleHeight, const FName& CategoryName, const FColor& Color = FColor::White, const FString& Description = TEXT(""), uint16 Thickness = 0);
+	// Cylinder
+	void AddElement(const FVector& Start, const FVector& End, float Radius, const FName& CategoryName, const FColor& Color = FColor::White, const FString& Description = TEXT(""), uint16 Thickness = 0);
+	// capsule
+	void AddElement(const FVector& Center, float HalfHeight, float Radius, const FQuat & Rotation, const FName& CategoryName, const FColor& Color = FColor::White, const FString& Description = TEXT(""));
+	// histogram sample
+	void AddHistogramData(const FVector2D& DataSample, const FName& CategoryName, const FName& GraphName, const FName& DataName);
+	// Custom data block
+	FVisualLogDataBlock& AddDataBlock(const FString& TagName, const TArray<uint8>& BlobDataArray, const FName& CategoryName);
+	// Event
+	int32 AddEvent(const FVisualLogEventBase& Event);
+	// find index of status category
+	int32 FindStatusIndex(const FString& CategoryName);
 
-		FORCEINLINE EElementType GetType() const
-		{
-			return EElementType(Type);
-		}
 
-		FORCEINLINE void SetType(EElementType InType)
-		{
-			Type = InType;
-		}
+#endif // ENABLE_VISUAL_LOG
+};
 
-		FORCEINLINE FColor GetFColor() const
-		{
-			return FColor(
-				((Color & 0xc0) << 24)
-				| ((Color & 0x30) << 18)
-				| ((Color & 0x0c) << 12)
-				| ((Color & 0x03) << 6)
-				);
-		}
-	};
+#if  ENABLE_VISUAL_LOG
 
-	struct FHistogramSample
-	{
-		FName Category;
-		TEnumAsByte<ELogVerbosity::Type> Verbosity;
-		FName GraphName;
-		FName DataName;
-		FVector2D SampleValue;
-		int32 UniqueId;
-	};
+class FVisualLogExtensionInterface
+{
+public:
+	virtual void OnTimestampChange(float Timestamp, class UWorld* InWorld, class AActor* HelperActor) = 0;
+	virtual void DrawData(class UWorld* InWorld, class UCanvas* Canvas, class AActor* HelperActor, const FName& TagName, const FVisualLogDataBlock& DataBlock, float Timestamp) = 0;
+	virtual void DisableDrawingForData(class UWorld* InWorld, class UCanvas* Canvas, class AActor* HelperActor, const FName& TagName, const FVisualLogDataBlock& DataBlock, float Timestamp) = 0;
+	virtual void LogEntryLineSelectionChanged(TSharedPtr<struct FLogEntryItem> SelectedItem, int64 UserData, FName TagName) = 0;
+};
 
-	struct FDataBlock
+/**
+ * Interface for Visual Logger Device
+ */
+class ENGINE_API FVisualLogDevice
+{
+public:
+	struct FVisualLogEntryItem
 	{
-		FName TagName;
-		FName Category;
-		TEnumAsByte<ELogVerbosity::Type> Verbosity;
-		TArray<uint8>	Data;
-		int32 UniqueId;
-	};
+		FVisualLogEntryItem() {}
+		FVisualLogEntryItem(FName InOwnerName, const FVisualLogEntry& LogEntry) : OwnerName(InOwnerName), Entry(LogEntry) { }
 
-#if ENABLE_VISUAL_LOG
-		float TimeStamp;
-		FVector Location;
-
-		TArray<FLogEvent> Events;
-		TArray<FLogLine> LogLines;
-		TArray<FStatusCategory> Status;
-		TArray<FElementToDraw> ElementsToDraw;
-		TArray<FHistogramSample>	HistogramSamples;
-		TArray<FDataBlock>	DataBlocks;
-
-		FVisualLogEntry() { Reset(); }
-		FVisualLogEntry(const FVisualLogEntry& Entry);
-		FVisualLogEntry(const class AActor* InActor, TArray<TWeakObjectPtr<UObject> >* Children);
-		FVisualLogEntry(float InTimeStamp, FVector InLocation, const UObject* Object, TArray<TWeakObjectPtr<UObject> >* Children);
-
-		void Reset()
-		{
-			TimeStamp = -1;
-			Location = FVector::ZeroVector;
-			Events.Reset();
-			LogLines.Reset();
-			Status.Reset();
-			ElementsToDraw.Reset();
-			HistogramSamples.Reset();
-			DataBlocks.Reset();
-		}
+		FName OwnerName;
+		FVisualLogEntry Entry;
+	};
 
-		void AddText(const FString& TextLine, const FName& CategoryName);
-		// path
-		void AddElement(const TArray<FVector>& Points, const FName& CategoryName, const FColor& Color = FColor::White, const FString& Description = TEXT(""), uint16 Thickness = 0);
-		// location
-		void AddElement(const FVector& Point, const FName& CategoryName, const FColor& Color = FColor::White, const FString& Description = TEXT(""), uint16 Thickness = 0);
-		// segment
-		void AddElement(const FVector& Start, const FVector& End, const FName& CategoryName, const FColor& Color = FColor::White, const FString& Description = TEXT(""), uint16 Thickness = 0);
-		// box
-		void AddElement(const FBox& Box, const FName& CategoryName, const FColor& Color = FColor::White, const FString& Description = TEXT(""), uint16 Thickness = 0);
-		// Cone
-		void AddElement(const FVector& Orgin, const FVector& Direction, float Length, float AngleWidth, float AngleHeight, const FName& CategoryName, const FColor& Color = FColor::White, const FString& Description = TEXT(""), uint16 Thickness = 0);
-		// Cylinder
-		void AddElement(const FVector& Start, const FVector& End, float Radius, const FName& CategoryName, const FColor& Color = FColor::White, const FString& Description = TEXT(""), uint16 Thickness = 0);
-		// histogram sample
-		void AddHistogramData(const FVector2D& DataSample, const FName& CategoryName, const FName& GraphName, const FName& DataName);
-		// Custom data block
-		FVisualLogEntry::FDataBlock& AddDataBlock(const FString& TagName, const TArray<uint8>& BlobDataArray, const FName& CategoryName);
-		// capsule
-		void AddCapsule(FVector const& Center, float HalfHeight, float Radius, const FQuat & Rotation, const FName& CategoryName, const FColor& Color = FColor::White, const FString& Description = TEXT(""));
-		// Event
-		int32 AddEvent(const FVisualLogEventBase& Event);
-
-		int32 FindStatusIndex(const FString& CategoryName);
-
-		// find index of status category
 
-#endif // ENABLE_VISUAL_LOG
+	virtual void Serialize(const class UObject* LogOwner, FName OwnerName, const FVisualLogEntry& LogEntry) = 0;
+	virtual void Cleanup(bool bReleaseMemory = false) { /* Empty */ }
+	virtual void StartRecordingToFile(float TImeStamp) { /* Empty */ }
+	virtual void StopRecordingToFile(float TImeStamp) { /* Empty */ }
+	virtual void SetFileName(const FString& InFileName) { /* Empty */ }
+	virtual void GetRecordedLogs(TArray<FVisualLogDevice::FVisualLogEntryItem>& OutLogs)  const { /* Empty */ }
+	virtual bool HasFlags(int32 InFlags) const { return !!(InFlags & EVisualLoggerDeviceFlags::NoFlags); }
 };
 
-#if  ENABLE_VISUAL_LOG
-FORCEINLINE
-FArchive& operator<<(FArchive& Ar, FVisualLogEntry::FDataBlock& Data)
+struct ENGINE_API FVisualLoggerHelpers
 {
-	Ar << Data.TagName;
-	Ar << Data.Category;
-	Ar << Data.Verbosity;
-	Ar << Data.Data;
-	Ar << Data.UniqueId;
+	static FString GenerateTemporaryFilename(const FString& FileExt);
+	static FString GenerateFilename(const FString& TempFileName, const FString& Prefix, float StartRecordingTime, float EndTimeStamp);
+	static FArchive& Serialize(FArchive& Ar, FName& Name);
+	static FArchive& Serialize(FArchive& Ar, TArray<FVisualLogDevice::FVisualLogEntryItem>& RecordedLogs);
+};
+
+ENGINE_API  FArchive& operator<<(FArchive& Ar, FVisualLogDevice::FVisualLogEntryItem& FrameCacheItem);
+ENGINE_API  FArchive& operator<<(FArchive& Ar, FVisualLogDataBlock& Data);
+ENGINE_API  FArchive& operator<<(FArchive& Ar, FVisualLogHistogramSample& Sample);
+ENGINE_API  FArchive& operator<<(FArchive& Ar, FVisualLogShapeElement& Element);
+ENGINE_API  FArchive& operator<<(FArchive& Ar, FVisualLogEvent& Event);
+ENGINE_API  FArchive& operator<<(FArchive& Ar, FVisualLogLine& LogLine);
+ENGINE_API  FArchive& operator<<(FArchive& Ar, FVisualLogStatusCategory& Status);
+ENGINE_API  FArchive& operator<<(FArchive& Ar, FVisualLogEntry& LogEntry);
 
-	return Ar;
+FORCEINLINE
+bool operator==(const FVisualLogEvent& Left, const FVisualLogEvent& Right) 
+{ 
+	return Left.Name == Right.Name; 
 }
 
 FORCEINLINE
-FArchive& operator<<(FArchive& Ar, FVisualLogEntry::FHistogramSample& Sample)
+FVisualLogEvent::FVisualLogEvent(const FVisualLogEventBase& Event)
+: Counter(1)
 {
-	Ar << Sample.Category;
-	Ar << Sample.Verbosity;
-	Ar << Sample.GraphName;
-	Ar << Sample.DataName;
-	Ar << Sample.SampleValue;
-	Ar << Sample.UniqueId;
-
-	return Ar;
+	Name = Event.Name;
+	UserFriendlyDesc = Event.FriendlyDesc;
+	Verbosity = Event.Verbosity;
 }
 
 FORCEINLINE
-FArchive& operator<<(FArchive& Ar, FVisualLogEntry::FElementToDraw& Element)
+FVisualLogEvent& FVisualLogEvent::operator= (const FVisualLogEventBase& Event)
 {
-	Ar << Element.Description;
-	Ar << Element.Category;
-	Ar << Element.Verbosity;
-	Ar << Element.Points;
-	Ar << Element.UniqueId;
-	Ar << Element.Type;
-	Ar << Element.Color;
-	Ar << Element.Thicknes;
-
-	return Ar;
+	Name = Event.Name;
+	UserFriendlyDesc = Event.FriendlyDesc;
+	Verbosity = Event.Verbosity;
+	return *this;
 }
 
 FORCEINLINE
-FArchive& operator<<(FArchive& Ar, FVisualLogEntry::FLogEvent& Event)
+FVisualLogLine::FVisualLogLine(const FName& InCategory, ELogVerbosity::Type InVerbosity, const FString& InLine)
+: Line(InLine)
+, Category(InCategory)
+, Verbosity(InVerbosity)
+, UserData(0)
 {
-	Ar << Event.Name;
-	Ar << Event.UserFriendlyDesc;
-	Ar << Event.Verbosity;
-	Ar << Event.EventTags;
-	Ar << Event.Counter;
-	Ar << Event.UserData;
-	Ar << Event.TagName;
-
-	return Ar;
+
 }
 
 FORCEINLINE
-FArchive& operator<<(FArchive& Ar, FVisualLogEntry::FLogLine& LogLine)
+FVisualLogLine::FVisualLogLine(const FName& InCategory, ELogVerbosity::Type InVerbosity, const FString& InLine, int64 InUserData)
+: Line(InLine)
+, Category(InCategory)
+, Verbosity(InVerbosity)
+, UserData(InUserData)
 {
-	Ar << LogLine.Category;
-	Ar << LogLine.Verbosity;
-	Ar << LogLine.TagName;
-	Ar << LogLine.UniqueId;
-	Ar << LogLine.UserData;
-	Ar << LogLine.Line;
-	return Ar;
+
 }
 
 FORCEINLINE
-FArchive& operator<<(FArchive& Ar, FVisualLogEntry::FStatusCategory& Status)
+void FVisualLogStatusCategory::Add(const FString& Key, const FString& Value)
 {
-	Ar << Status.Category;
-	Ar << Status.Data;
-
-	return Ar;
+	Data.Add(FString(Key).AppendChar(TEXT('|')) + Value);
 }
 
 FORCEINLINE
-FArchive& operator<<(FArchive& Ar, FVisualLogEntry& LogEntry)
+FVisualLogShapeElement::FVisualLogShapeElement(EVisualLoggerShapeElement InType)
+: Verbosity(ELogVerbosity::All)
+, Type(InType)
+, Color(0xff)
+, Thicknes(0)
 {
-	Ar << LogEntry.TimeStamp;
-	Ar << LogEntry.Location;
-	Ar << LogEntry.LogLines;
-	Ar << LogEntry.Status;
-	Ar << LogEntry.Events;
-	Ar << LogEntry.ElementsToDraw;
-
-	return Ar;
+
 }
 
-FORCEINLINE int32 FVisualLogEntry::FindStatusIndex(const FString& CategoryName)
+FORCEINLINE
+FVisualLogShapeElement::FVisualLogShapeElement(const FString& InDescription, const FColor& InColor, uint16 InThickness, const FName& InCategory)
+: Category(InCategory)
+, Verbosity(ELogVerbosity::All)
+, Type(EVisualLoggerShapeElement::Invalid)
+, Thicknes(InThickness)
 {
-	for (int32 TestCategoryIndex = 0; TestCategoryIndex < Status.Num(); TestCategoryIndex++)
+	if (InDescription.IsEmpty() == false)
 	{
-		if (Status[TestCategoryIndex].Category == CategoryName)
-		{
-			return TestCategoryIndex;
-		}
+		Description = InDescription;
 	}
+	SetColor(InColor);
+}
 
-	return INDEX_NONE;
+FORCEINLINE
+void FVisualLogShapeElement::SetColor(const FColor& InColor)
+{
+	Color = ((InColor.DWColor() >> 30) << 6)	| (((InColor.DWColor() & 0x00ff0000) >> 22) << 4)	| (((InColor.DWColor() & 0x0000ff00) >> 14) << 2)	| ((InColor.DWColor() & 0x000000ff) >> 6);
 }
 
-FORCEINLINE bool operator==(const FVisualLogEntry::FLogEvent& Left, const FVisualLogEntry::FLogEvent& Right)
+FORCEINLINE
+EVisualLoggerShapeElement FVisualLogShapeElement::GetType() const
 {
-	return Left.Name == Right.Name;
+	return Type;
 }
 
-struct FVisualLoggerHelpers
+FORCEINLINE
+void FVisualLogShapeElement::SetType(EVisualLoggerShapeElement InType)
 {
-	static FString GenerateTemporaryFilename(const FString& FileExt);
-	static FString GenerateFilename(const FString& TempFileName, const FString& Prefix, float StartRecordingTime, float EndTimeStamp);
-};
+	Type = InType;
+}
 
-FORCEINLINE 
-FString FVisualLoggerHelpers::GenerateTemporaryFilename(const FString& FileExt)
-{ 
-	return FString::Printf(TEXT("VTEMP_%s.%s"), *FDateTime::Now().ToString(), *FileExt);
+FORCEINLINE
+FColor FVisualLogShapeElement::GetFColor() const
+{
+	return FColor(((Color & 0xc0) << 24) | ((Color & 0x30) << 18) | ((Color & 0x0c) << 12) | ((Color & 0x03) << 6));
 }
 
+
 FORCEINLINE
-FString FVisualLoggerHelpers::GenerateFilename(const FString& TempFileName, const FString& Prefix, float StartRecordingTime, float EndTimeStamp)
+int32 FVisualLogEntry::FindStatusIndex(const FString& CategoryName)
 {
-	const FString TempFullFilename = FString::Printf(TEXT("%slogs/%s"), *FPaths::GameSavedDir(), *TempFileName);
-	const FString FullFilename = FString::Printf(TEXT("%slogs/%s_%s"), *FPaths::GameSavedDir(), *Prefix, *TempFileName);
-	const FString TimeFrameString = FString::Printf(TEXT("%d-%d_"), FMath::TruncToInt(StartRecordingTime), FMath::TruncToInt(EndTimeStamp));
-	return FullFilename.Replace(TEXT("VTEMP_"), *TimeFrameString, ESearchCase::CaseSensitive);
+	for (int32 TestCategoryIndex = 0; TestCategoryIndex < Status.Num(); TestCategoryIndex++)
+	{
+		if (Status[TestCategoryIndex].Category == CategoryName)
+		{
+			return TestCategoryIndex;
+		}
+	}
+
+	return INDEX_NONE;
 }
 
 #endif // ENABLE_VISUAL_LOG
