Fix case-sensitive error.

diff --git a/Engine/Source/Developer/LogVisualizer/Private/SLogVisualizer.cpp b/Engine/Source/Developer/LogVisualizer/Private/SLogVisualizer.cpp
--- a/Engine/Source/Developer/LogVisualizer/Private/SLogVisualizer.cpp
+++ b/Engine/Source/Developer/LogVisualizer/Private/SLogVisualizer.cpp
@@ -19,7 +19,7 @@
 #	include "Editor/UnrealEd/Public/EditorViewportClient.h"
 #endif
 
-#include "VIsualLogger/VisualLoggerBinaryFileDevice.h"
+#include "VisualLogger/VisualLoggerBinaryFileDevice.h"
 #include "GameplayDebuggingComponent.h"
 #include "LogVisualizerModule.h"
 #include "SLogVisualizerReport.h"
