LogVisualizer - fixed toolbar icons #ue4

diff --git a/Engine/Source/Developer/LogVisualizer/Private/SLogVisualizer.cpp b/Engine/Source/Developer/LogVisualizer/Private/SLogVisualizer.cpp
--- a/Engine/Source/Developer/LogVisualizer/Private/SLogVisualizer.cpp
+++ b/Engine/Source/Developer/LogVisualizer/Private/SLogVisualizer.cpp
@@ -162,8 +162,9 @@ void SLogVisualizer::Construct(const FArguments& InArgs, FLogVisualizer* InLogVi
 					SNew(SHorizontalBox)
 					// Record button
 					+SHorizontalBox::Slot()
-					.AutoWidth()
 					.Padding(1)
+					.AutoWidth()
+					.VAlign(VAlign_Center)
 					[
 						SNew(SButton)
 						.OnClicked(this, &SLogVisualizer::OnRecordButtonClicked)
@@ -175,11 +176,13 @@ void SLogVisualizer::Construct(const FArguments& InArgs, FLogVisualizer* InLogVi
 					]
 					// 'Pause' toggle button
 					+SHorizontalBox::Slot()
-					.AutoWidth()
 					.Padding(1)
+					.AutoWidth()
+					.VAlign(VAlign_Center)
 					[
 						SNew(SCheckBox)
 						.Style(FEditorStyle::Get(), "ToggleButtonCheckbox")
+						.Type(ESlateCheckBoxType::ToggleButton)
 						.OnCheckStateChanged(this, &SLogVisualizer::OnPauseChanged)
 						.IsChecked(this, &SLogVisualizer::GetPauseState)
 						.Content()
@@ -190,8 +193,9 @@ void SLogVisualizer::Construct(const FArguments& InArgs, FLogVisualizer* InLogVi
 					]
 					// 'Camera' toggle button
 					+SHorizontalBox::Slot()
-					.AutoWidth()
 					.Padding(1)
+					.AutoWidth()
+					.VAlign(VAlign_Center)
 					[
 						SNew(SCheckBox)
 						.Style(FEditorStyle::Get(), "ToggleButtonCheckbox")
@@ -212,8 +216,9 @@ void SLogVisualizer::Construct(const FArguments& InArgs, FLogVisualizer* InLogVi
 					]
 					// 'Save' function
 					+SHorizontalBox::Slot()
-					.AutoWidth()
 					.Padding(1)
+					.AutoWidth()
+					.VAlign(VAlign_Center)
 					[
 						SNew(SButton)
 						.OnClicked(this, &SLogVisualizer::OnSave)
@@ -225,8 +230,9 @@ void SLogVisualizer::Construct(const FArguments& InArgs, FLogVisualizer* InLogVi
 					]
 					// 'Load' function
 					+SHorizontalBox::Slot()
-					.AutoWidth()
 					.Padding(1)
+					.AutoWidth()
+					.VAlign(VAlign_Center)
 					[
 						SNew(SButton)
 						.OnClicked(this, &SLogVisualizer::OnLoad)
@@ -238,8 +244,9 @@ void SLogVisualizer::Construct(const FArguments& InArgs, FLogVisualizer* InLogVi
 					]
 					// 'Remove' function
 					+SHorizontalBox::Slot()
-					.AutoWidth()
 					.Padding(1)
+					.AutoWidth()
+					.VAlign(VAlign_Center)
 					[
 						SNew(SButton)
 						.OnClicked(this, &SLogVisualizer::OnRemove)
diff --git a/Engine/Source/Editor/EditorStyle/Private/SlateEditorStyle.cpp b/Engine/Source/Editor/EditorStyle/Private/SlateEditorStyle.cpp
--- a/Engine/Source/Editor/EditorStyle/Private/SlateEditorStyle.cpp
+++ b/Engine/Source/Editor/EditorStyle/Private/SlateEditorStyle.cpp
@@ -1537,14 +1537,14 @@ void FSlateEditorStyle::FStyle::SetupGeneralStyles()
 
 	// Log Visualizer
 	{
-		Set( "LogVisualizer.Pause", new IMAGE_BRUSH( "Icons/icon_pause_40x", Icon40x40 ) );
+		Set("LogVisualizer.Pause", new IMAGE_BRUSH("Icons/icon_pause_40x", Icon48x48));
 		Set( "LogVisualizer.LogBar.Background", new BOX_BRUSH( "Common/ProgressBar_Background", FMargin(5.f/12.f) ) );
 		Set( "LogVisualizer.LogBar.Selected", new BOX_BRUSH( "Common/TaskGraph_Selected", FMargin(5.f/12.f) ) );
 		Set( "LogVisualizer.LogBar.EntryDefault", new BOX_BRUSH( "Common/TaskGraph_Mono", FMargin(5.f/12.f) ) );
 		Set( "LogVisualizer.LogBar.TimeMark", new BOX_BRUSH( "Icons/LV_BarMark", FMargin(5.f/12.f) ) );
-		Set( "LogVisualizer.Record", new IMAGE_BRUSH( "Icons/LV_Record", Icon24x24 ) );
-		Set( "LogVisualizer.Stop", new IMAGE_BRUSH( "Icons/LV_Stop", Icon24x24 ) );
-		Set( "LogVisualizer.Camera", new IMAGE_BRUSH( "Icons/LV_Camera", Icon40x40 ) );
+		Set("LogVisualizer.Record", new IMAGE_BRUSH("Icons/LV_Record", Icon40x40));
+		Set("LogVisualizer.Stop", new IMAGE_BRUSH("Icons/LV_Stop", Icon40x40));
+		Set("LogVisualizer.Camera", new IMAGE_BRUSH("Icons/LV_Camera", Icon48x48));
 		Set( "LogVisualizer.Save", new IMAGE_BRUSH( "Icons/LV_Save", Icon40x40 ) );
 		Set( "LogVisualizer.Load", new IMAGE_BRUSH( "Icons/LV_Load", Icon40x40 ) );
 		Set( "LogVisualizer.Remove", new IMAGE_BRUSH( "Icons/LV_Remove", Icon40x40 ) );
