Slate - ESlateCheckBoxState has been renamed to ECheckBoxState.

diff --git a/Engine/Config/BaseEngine.ini b/Engine/Config/BaseEngine.ini
--- a/Engine/Config/BaseEngine.ini
+++ b/Engine/Config/BaseEngine.ini
@@ -1097,6 +1097,9 @@ bUseStreamingPause=false
 +K2ParamRedirects=(NodeName="/Script/BlueprintGraph.K2Node_FunctionEntry", OldParamName="UserWidget.OnControllerButtonPressed.ControllerEvent", NewParamName="InKeyEvent")
 +K2ParamRedirects=(NodeName="/Script/BlueprintGraph.K2Node_FunctionEntry", OldParamName="UserWidget.OnControllerButtonReleased.ControllerEvent", NewParamName="InKeyEvent")
 
+; UMG Checkbox State Rename
+
++ActiveClassRedirects=(OldClassName="ESlateCheckBoxState", NewClassName="ECheckBoxState")
 
 
 ; VER_UE4_INVERSE_SQUARED_LIGHTS_DEFAULT
diff --git a/Engine/Plugins/Editor/EpicSurvey/Source/EpicSurvey/Private/SQuestionBlock.cpp b/Engine/Plugins/Editor/EpicSurvey/Source/EpicSurvey/Private/SQuestionBlock.cpp
--- a/Engine/Plugins/Editor/EpicSurvey/Source/EpicSurvey/Private/SQuestionBlock.cpp
+++ b/Engine/Plugins/Editor/EpicSurvey/Source/EpicSurvey/Private/SQuestionBlock.cpp
@@ -302,7 +302,7 @@ void SQuestionBlock::Construct( const FArguments& Args, const TSharedRef<FQuesti
 	];
 }
 
-ESlateCheckBoxState::Type SQuestionBlock::IsAnswerChecked( TWeakPtr< FQuestionBlock > BlockPtr, int32 QuestionIndex, int32 AnswerIndex ) const
+ECheckBoxState SQuestionBlock::IsAnswerChecked( TWeakPtr< FQuestionBlock > BlockPtr, int32 QuestionIndex, int32 AnswerIndex ) const
 {
 	TSharedPtr< FQuestionBlock > Block =  BlockPtr.Pin();
 
@@ -312,10 +312,10 @@ ESlateCheckBoxState::Type SQuestionBlock::IsAnswerChecked( TWeakPtr< FQuestionBl
 		IsChecked = Block->GetUserAnswers( QuestionIndex ).Contains( AnswerIndex );
 	}
 
-	return IsChecked ? ESlateCheckBoxState::Checked : ESlateCheckBoxState::Unchecked;
+	return IsChecked ? ECheckBoxState::Checked : ECheckBoxState::Unchecked;
 }
 
-void SQuestionBlock::AnswerCheckStateChanged( ESlateCheckBoxState::Type CheckState, TWeakPtr< FQuestionBlock > BlockPtr, int32 QuestionIndex, int32 AnswerIndex )
+void SQuestionBlock::AnswerCheckStateChanged( ECheckBoxState CheckState, TWeakPtr< FQuestionBlock > BlockPtr, int32 QuestionIndex, int32 AnswerIndex )
 {
 	TSharedPtr< FQuestionBlock > Block =  BlockPtr.Pin();
 
diff --git a/Engine/Plugins/Editor/EpicSurvey/Source/EpicSurvey/Private/SQuestionBlock.h b/Engine/Plugins/Editor/EpicSurvey/Source/EpicSurvey/Private/SQuestionBlock.h
--- a/Engine/Plugins/Editor/EpicSurvey/Source/EpicSurvey/Private/SQuestionBlock.h
+++ b/Engine/Plugins/Editor/EpicSurvey/Source/EpicSurvey/Private/SQuestionBlock.h
@@ -17,8 +17,8 @@ public:
 
 private:
 
-	ESlateCheckBoxState::Type IsAnswerChecked( TWeakPtr< FQuestionBlock > BlockPtr, int32 QuestionIndex, int32 AnswerIndex ) const;
-	void AnswerCheckStateChanged( ESlateCheckBoxState::Type CheckState, TWeakPtr< FQuestionBlock > BlockPtr, int32 QuestionIndex, int32 AnswerIndex );
+	ECheckBoxState IsAnswerChecked( TWeakPtr< FQuestionBlock > BlockPtr, int32 QuestionIndex, int32 AnswerIndex ) const;
+	void AnswerCheckStateChanged( ECheckBoxState CheckState, TWeakPtr< FQuestionBlock > BlockPtr, int32 QuestionIndex, int32 AnswerIndex );
 
 	void AnswerTextChanged( const FText& Text, TWeakPtr< FQuestionBlock > BlockPtr, int32 QuestionIndex, int32 AnswerIndex ) const;
 	void AnswerTextCommitted( const FText& Text, ETextCommit::Type CommitType, TWeakPtr< FQuestionBlock > BlockPtr, int32 QuestionIndex, int32 AnswerIndex ) const;
diff --git a/Engine/Plugins/Editor/EpicSurvey/Source/EpicSurvey/Private/SSurvey.h b/Engine/Plugins/Editor/EpicSurvey/Source/EpicSurvey/Private/SSurvey.h
--- a/Engine/Plugins/Editor/EpicSurvey/Source/EpicSurvey/Private/SSurvey.h
+++ b/Engine/Plugins/Editor/EpicSurvey/Source/EpicSurvey/Private/SSurvey.h
@@ -25,9 +25,9 @@ private:
 	void ConstructSurveyLayout();
 	void ConstructFailureLayout();
 
-	ESlateCheckBoxState::Type IsAnswerChecked( TWeakPtr< FQuestionBlock > BlockPtr, int32 QuestionIndex, int32 AnswerIndex ) const;
+	ECheckBoxState IsAnswerChecked( TWeakPtr< FQuestionBlock > BlockPtr, int32 QuestionIndex, int32 AnswerIndex ) const;
 
-	void AnswerCheckStateChanged( ESlateCheckBoxState::Type CheckState, TWeakPtr< FQuestionBlock > BlockPtr, int32 QuestionIndex, int32 AnswerIndex );
+	void AnswerCheckStateChanged( ECheckBoxState CheckState, TWeakPtr< FQuestionBlock > BlockPtr, int32 QuestionIndex, int32 AnswerIndex );
 	EVisibility CanSubmitSurvey() const;
 	FReply SubmitSurvey();
 
diff --git a/Engine/Plugins/Editor/PluginsEditor/Source/PluginsEditor/Private/SPluginListTile.cpp b/Engine/Plugins/Editor/PluginsEditor/Source/PluginsEditor/Private/SPluginListTile.cpp
--- a/Engine/Plugins/Editor/PluginsEditor/Source/PluginsEditor/Private/SPluginListTile.cpp
+++ b/Engine/Plugins/Editor/PluginsEditor/Source/PluginsEditor/Private/SPluginListTile.cpp
@@ -301,17 +301,17 @@ void SPluginListTile::Construct( const FArguments& Args, const TSharedRef<SPlugi
 }
 
 
-ESlateCheckBoxState::Type SPluginListTile::IsPluginEnabled() const
+ECheckBoxState SPluginListTile::IsPluginEnabled() const
 {
 	return ItemData->PluginStatus.bIsEnabled
-		? ESlateCheckBoxState::Checked
-		: ESlateCheckBoxState::Unchecked;
+		? ECheckBoxState::Checked
+		: ECheckBoxState::Unchecked;
 }
 
 
-void SPluginListTile::OnEnablePluginCheckboxChanged(ESlateCheckBoxState::Type NewCheckedState)
+void SPluginListTile::OnEnablePluginCheckboxChanged(ECheckBoxState NewCheckedState)
 {
-	const bool bNewEnabledState = (NewCheckedState == ESlateCheckBoxState::Checked);
+	const bool bNewEnabledState = (NewCheckedState == ECheckBoxState::Checked);
 
 	if (bNewEnabledState && ItemData->PluginStatus.bIsBetaVersion)
 	{
diff --git a/Engine/Plugins/Editor/PluginsEditor/Source/PluginsEditor/Private/SPluginListTile.h b/Engine/Plugins/Editor/PluginsEditor/Source/PluginsEditor/Private/SPluginListTile.h
--- a/Engine/Plugins/Editor/PluginsEditor/Source/PluginsEditor/Private/SPluginListTile.h
+++ b/Engine/Plugins/Editor/PluginsEditor/Source/PluginsEditor/Private/SPluginListTile.h
@@ -21,10 +21,10 @@ class SPluginListTile : public SCompoundWidget
 private:
 
 	/** Returns the checked state for the enabled checkbox */
-	ESlateCheckBoxState::Type IsPluginEnabled() const;
+	ECheckBoxState IsPluginEnabled() const;
 
 	/** Called when the enabled checkbox is clicked */
-	void OnEnablePluginCheckboxChanged(ESlateCheckBoxState::Type NewCheckedState);
+	void OnEnablePluginCheckboxChanged(ECheckBoxState NewCheckedState);
 
 	/** Called when 'Uninstall' is clicked on a plugin's tile */
 	FReply OnUninstallButtonClicked();
diff --git a/Engine/Plugins/Editor/SpeedTreeImporter/Source/SpeedTreeImporter/Private/SpeedTreeImportFactory.cpp b/Engine/Plugins/Editor/SpeedTreeImporter/Source/SpeedTreeImporter/Private/SpeedTreeImportFactory.cpp
--- a/Engine/Plugins/Editor/SpeedTreeImporter/Source/SpeedTreeImporter/Private/SpeedTreeImportFactory.cpp
+++ b/Engine/Plugins/Editor/SpeedTreeImporter/Source/SpeedTreeImporter/Private/SpeedTreeImportFactory.cpp
@@ -103,7 +103,7 @@ public:
 	TSharedRef<SWidget> CreateOption(TSharedPtr<SCheckBox>& CheckBox, const FText& NameText, bool bChecked)
 	{
 		return SAssignNew(CheckBox, SCheckBox)
-		.IsChecked(bChecked ? ESlateCheckBoxState::Checked : ESlateCheckBoxState::Unchecked)
+		.IsChecked(bChecked ? ECheckBoxState::Checked : ECheckBoxState::Unchecked)
 		.OnCheckStateChanged(this, &SSpeedTreeImportOptions::OnOptionModified)
 		.Content()
 		[
@@ -123,14 +123,14 @@ public:
 		];
 	}
 
-	ESlateCheckBoxState::Type IsRadioChecked(EImportGeometryType ButtonId) const
+	ECheckBoxState IsRadioChecked(EImportGeometryType ButtonId) const
 	{
-		return (ImportGeometryType == ButtonId) ? ESlateCheckBoxState::Checked : ESlateCheckBoxState::Unchecked;
+		return (ImportGeometryType == ButtonId) ? ECheckBoxState::Checked : ECheckBoxState::Unchecked;
 	}
 
-	void OnRadioChanged( ESlateCheckBoxState::Type NewRadioState, EImportGeometryType RadioThatChanged )
+	void OnRadioChanged( ECheckBoxState NewRadioState, EImportGeometryType RadioThatChanged )
 	{
-		if (NewRadioState == ESlateCheckBoxState::Checked)
+		if (NewRadioState == ECheckBoxState::Checked)
 		{
 			ImportGeometryType = RadioThatChanged;
 		}
@@ -312,7 +312,7 @@ public:
 		return FReply::Handled();
 	}
 
-	void OnOptionModified(const ESlateCheckBoxState::Type /*NewCheckedState*/)
+	void OnOptionModified(const ECheckBoxState /*NewCheckedState*/)
 	{
 		bool bEnable = MakeMaterialsCheck->IsChecked();
 
diff --git a/Engine/Plugins/Experimental/LiveEditor/Source/LiveEditor/Private/LiveEditorDeviceSetupWizard.cpp b/Engine/Plugins/Experimental/LiveEditor/Source/LiveEditor/Private/LiveEditorDeviceSetupWizard.cpp
--- a/Engine/Plugins/Experimental/LiveEditor/Source/LiveEditor/Private/LiveEditorDeviceSetupWizard.cpp
+++ b/Engine/Plugins/Experimental/LiveEditor/Source/LiveEditor/Private/LiveEditorDeviceSetupWizard.cpp
@@ -34,10 +34,10 @@ public:
 public:
 	void Construct(const FArguments& InArgs, struct FConfigurationState *_Owner);
 
-	void OnHasButtonsChanged(ESlateCheckBoxState::Type NewCheckedState);
-	ESlateCheckBoxState::Type HasButtons() const;
-	void OnHasEndlessEncodersChanged(ESlateCheckBoxState::Type NewCheckedState);
-	ESlateCheckBoxState::Type HasEndlessEncoders() const;
+	void OnHasButtonsChanged(ECheckBoxState NewCheckedState);
+	ECheckBoxState HasButtons() const;
+	void OnHasEndlessEncodersChanged(ECheckBoxState NewCheckedState);
+	ECheckBoxState HasEndlessEncoders() const;
 
 private:
 	struct FConfigurationState *Owner;
@@ -143,25 +143,25 @@ void SDeviceQuestionnaireSubWindow::Construct(const FArguments& InArgs, FConfigu
 		]
 	];
 }
-void SDeviceQuestionnaireSubWindow::OnHasButtonsChanged(ESlateCheckBoxState::Type NewCheckedState)
+void SDeviceQuestionnaireSubWindow::OnHasButtonsChanged(ECheckBoxState NewCheckedState)
 {
 	check( Owner != NULL );
-	Owner->bHasButtons = (NewCheckedState == ESlateCheckBoxState::Checked);
+	Owner->bHasButtons = (NewCheckedState == ECheckBoxState::Checked);
 }
-ESlateCheckBoxState::Type SDeviceQuestionnaireSubWindow::HasButtons() const
+ECheckBoxState SDeviceQuestionnaireSubWindow::HasButtons() const
 {
 	check( Owner != NULL );
-	return (Owner->bHasButtons)? ESlateCheckBoxState::Checked : ESlateCheckBoxState::Unchecked;
+	return (Owner->bHasButtons)? ECheckBoxState::Checked : ECheckBoxState::Unchecked;
 }
-void SDeviceQuestionnaireSubWindow::OnHasEndlessEncodersChanged(ESlateCheckBoxState::Type NewCheckedState)
+void SDeviceQuestionnaireSubWindow::OnHasEndlessEncodersChanged(ECheckBoxState NewCheckedState)
 {
 	check( Owner != NULL );
-	Owner->bHasEndlessEncoders = (NewCheckedState == ESlateCheckBoxState::Checked);
+	Owner->bHasEndlessEncoders = (NewCheckedState == ECheckBoxState::Checked);
 }
-ESlateCheckBoxState::Type SDeviceQuestionnaireSubWindow::HasEndlessEncoders() const
+ECheckBoxState SDeviceQuestionnaireSubWindow::HasEndlessEncoders() const
 {
 	check( Owner != NULL );
-	return (Owner->bHasEndlessEncoders)? ESlateCheckBoxState::Checked : ESlateCheckBoxState::Unchecked;
+	return (Owner->bHasEndlessEncoders)? ECheckBoxState::Checked : ECheckBoxState::Unchecked;
 }
 
 
diff --git a/Engine/Plugins/Messaging/MessagingDebugger/Source/MessagingDebugger/Private/Widgets/Breakpoints/SMessagingBreakpointsTableRow.h b/Engine/Plugins/Messaging/MessagingDebugger/Source/MessagingDebugger/Private/Widgets/Breakpoints/SMessagingBreakpointsTableRow.h
--- a/Engine/Plugins/Messaging/MessagingDebugger/Source/MessagingDebugger/Private/Widgets/Breakpoints/SMessagingBreakpointsTableRow.h
+++ b/Engine/Plugins/Messaging/MessagingDebugger/Source/MessagingDebugger/Private/Widgets/Breakpoints/SMessagingBreakpointsTableRow.h
@@ -56,7 +56,7 @@ public:
 				+ SHorizontalBox::Slot()
 					[
 						SNew(SCheckBox)
-							.IsChecked(ESlateCheckBoxState::Checked)
+							.IsChecked(ECheckBoxState::Checked)
 					]
 
 				+ SHorizontalBox::Slot()
diff --git a/Engine/Plugins/Messaging/MessagingDebugger/Source/MessagingDebugger/Private/Widgets/Endpoints/SMessagingEndpointsTableRow.h b/Engine/Plugins/Messaging/MessagingDebugger/Source/MessagingDebugger/Private/Widgets/Endpoints/SMessagingEndpointsTableRow.h
--- a/Engine/Plugins/Messaging/MessagingDebugger/Source/MessagingDebugger/Private/Widgets/Endpoints/SMessagingEndpointsTableRow.h
+++ b/Engine/Plugins/Messaging/MessagingDebugger/Source/MessagingDebugger/Private/Widgets/Endpoints/SMessagingEndpointsTableRow.h
@@ -125,20 +125,20 @@ private:
 	}
 
 	/** Handles changing the checked state of the visibility check box. */
-	void HandleVisibilityCheckBoxCheckStateChanged( ESlateCheckBoxState::Type CheckState )
+	void HandleVisibilityCheckBoxCheckStateChanged( ECheckBoxState CheckState )
 	{
-		Model->SetEndpointVisibility(EndpointInfo.ToSharedRef(), (CheckState == ESlateCheckBoxState::Checked));
+		Model->SetEndpointVisibility(EndpointInfo.ToSharedRef(), (CheckState == ECheckBoxState::Checked));
 	}
 
 	/** Gets the image for the visibility check box. */
-	ESlateCheckBoxState::Type HandleVisibilityCheckBoxIsChecked() const
+	ECheckBoxState HandleVisibilityCheckBoxIsChecked() const
 	{
 		if (Model->IsEndpointVisible(EndpointInfo.ToSharedRef()))
 		{
-			return ESlateCheckBoxState::Checked;
+			return ECheckBoxState::Checked;
 		}
 
-		return ESlateCheckBoxState::Unchecked;
+		return ECheckBoxState::Unchecked;
 	}
 
 private:
diff --git a/Engine/Plugins/Messaging/MessagingDebugger/Source/MessagingDebugger/Private/Widgets/Types/SMessagingTypesTableRow.h b/Engine/Plugins/Messaging/MessagingDebugger/Source/MessagingDebugger/Private/Widgets/Types/SMessagingTypesTableRow.h
--- a/Engine/Plugins/Messaging/MessagingDebugger/Source/MessagingDebugger/Private/Widgets/Types/SMessagingTypesTableRow.h
+++ b/Engine/Plugins/Messaging/MessagingDebugger/Source/MessagingDebugger/Private/Widgets/Types/SMessagingTypesTableRow.h
@@ -106,20 +106,20 @@ private:
 	}
 
 	/** Handles changing the checked state of the visibility check box. */
-	void HandleVisibilityCheckBoxCheckStateChanged( ESlateCheckBoxState::Type CheckState )
+	void HandleVisibilityCheckBoxCheckStateChanged( ECheckBoxState CheckState )
 	{
-		Model->SetTypeVisibility(TypeInfo.ToSharedRef(), (CheckState == ESlateCheckBoxState::Checked));
+		Model->SetTypeVisibility(TypeInfo.ToSharedRef(), (CheckState == ECheckBoxState::Checked));
 	}
 
 	/** Gets the image for the visibility check box. */
-	ESlateCheckBoxState::Type HandleVisibilityCheckBoxIsChecked() const
+	ECheckBoxState HandleVisibilityCheckBoxIsChecked() const
 	{
 		if (Model->IsTypeVisible(TypeInfo.ToSharedRef()))
 		{
-			return ESlateCheckBoxState::Checked;
+			return ECheckBoxState::Checked;
 		}
 
-		return ESlateCheckBoxState::Unchecked;
+		return ECheckBoxState::Unchecked;
 	}
 
 private:
diff --git a/Engine/Source/Developer/AssetTools/Private/AssetTypeActions/AssetTypeActions_SkeletalMesh.cpp b/Engine/Source/Developer/AssetTools/Private/AssetTypeActions/AssetTypeActions_SkeletalMesh.cpp
--- a/Engine/Source/Developer/AssetTools/Private/AssetTypeActions/AssetTypeActions_SkeletalMesh.cpp
+++ b/Engine/Source/Developer/AssetTools/Private/AssetTypeActions/AssetTypeActions_SkeletalMesh.cpp
@@ -204,9 +204,9 @@ public:
 	 * @param	ButtonId	The ID for the check box
 	 * @return				The status of the check box
 	 */
-	ESlateCheckBoxState::Type IsCheckboxChecked( int32 ButtonId ) const
+	ECheckBoxState IsCheckboxChecked( int32 ButtonId ) const
 	{
-		return CheckBoxInfoMap.FindChecked(ButtonId).bUsed ? ESlateCheckBoxState::Checked : ESlateCheckBoxState::Unchecked;
+		return CheckBoxInfoMap.FindChecked(ButtonId).bUsed ? ECheckBoxState::Checked : ECheckBoxState::Unchecked;
 	}
 
 	/**
@@ -215,7 +215,7 @@ public:
 	 * @param	NewCheckboxState	The new state of the check box
 	 * @param	CheckboxThatChanged	The ID of the radio button that has changed. 
 	 */
-	void OnCheckboxChanged( ESlateCheckBoxState::Type NewCheckboxState, int32 CheckboxThatChanged )
+	void OnCheckboxChanged( ECheckBoxState NewCheckboxState, int32 CheckboxThatChanged )
 	{
 		FBoneCheckboxInfo& Info = CheckBoxInfoMap.FindChecked(CheckboxThatChanged);
 		Info.bUsed = !Info.bUsed;
diff --git a/Engine/Source/Developer/AssetTools/Private/AssetTypeActions/AssetTypeActions_Skeleton.cpp b/Engine/Source/Developer/AssetTools/Private/AssetTypeActions/AssetTypeActions_Skeleton.cpp
--- a/Engine/Source/Developer/AssetTools/Private/AssetTypeActions/AssetTypeActions_Skeleton.cpp
+++ b/Engine/Source/Developer/AssetTools/Private/AssetTypeActions/AssetTypeActions_Skeleton.cpp
@@ -203,9 +203,9 @@ public:
 	 * @param	ButtonId	The ID for the check box
 	 * @return				The status of the check box
 	 */
-	ESlateCheckBoxState::Type IsCheckboxChecked( int32 ButtonId ) const
+	ECheckBoxState IsCheckboxChecked( int32 ButtonId ) const
 	{
-		return CheckBoxInfoMap.FindChecked(ButtonId).bUsed ? ESlateCheckBoxState::Checked : ESlateCheckBoxState::Unchecked;
+		return CheckBoxInfoMap.FindChecked(ButtonId).bUsed ? ECheckBoxState::Checked : ECheckBoxState::Unchecked;
 	}
 
 	/**
@@ -214,7 +214,7 @@ public:
 	 * @param	NewCheckboxState	The new state of the check box
 	 * @param	CheckboxThatChanged	The ID of the radio button that has changed. 
 	 */
-	void OnCheckboxChanged( ESlateCheckBoxState::Type NewCheckboxState, int32 CheckboxThatChanged )
+	void OnCheckboxChanged( ECheckBoxState NewCheckboxState, int32 CheckboxThatChanged )
 	{
 		FBoneCheckbox& Info = CheckBoxInfoMap.FindChecked(CheckboxThatChanged);
 		Info.bUsed = !Info.bUsed;
diff --git a/Engine/Source/Developer/AutomationWindow/Private/SAutomationExportMenu.cpp b/Engine/Source/Developer/AutomationWindow/Private/SAutomationExportMenu.cpp
--- a/Engine/Source/Developer/AutomationWindow/Private/SAutomationExportMenu.cpp
+++ b/Engine/Source/Developer/AutomationWindow/Private/SAutomationExportMenu.cpp
@@ -217,10 +217,10 @@ bool SAutomationExportMenu::IsExportReady() const
 }
 
 
-void SAutomationExportMenu::OnDisplayCheckStateChanged( ESlateCheckBoxState::Type InNewState, EFileExportType::Type CheckType )
+void SAutomationExportMenu::OnDisplayCheckStateChanged( ECheckBoxState InNewState, EFileExportType::Type CheckType )
 {
 	// set or unset the bit in the mask for the type changed
-	if ( InNewState == ESlateCheckBoxState::Checked )
+	if ( InNewState == ECheckBoxState::Checked )
 	{
 		EFileExportType::SetFlag( FileExportTypeMask, CheckType );
 	}
@@ -238,7 +238,7 @@ void SAutomationExportMenu::OnDisplayCheckStateChanged( ESlateCheckBoxState::Typ
 	}
 
 	// If we check select all, set all the valid options to be active
-	if ( CheckType == EFileExportType::FET_All && InNewState == ESlateCheckBoxState::Checked )
+	if ( CheckType == EFileExportType::FET_All && InNewState == ECheckBoxState::Checked )
 	{
 		EnableAvailableReports();
 		// Recreate the menu
@@ -247,14 +247,14 @@ void SAutomationExportMenu::OnDisplayCheckStateChanged( ESlateCheckBoxState::Typ
 }
 
 
-ESlateCheckBoxState::Type SAutomationExportMenu::OnGetDisplayCheckState( EFileExportType::Type CheckType ) const
+ECheckBoxState SAutomationExportMenu::OnGetDisplayCheckState( EFileExportType::Type CheckType ) const
 {
 	// If the item is selected, set the box to be checked
 	if(EFileExportType::IsSet( FileExportTypeMask, CheckType ) )
 	{
-		return ESlateCheckBoxState::Checked;
+		return ECheckBoxState::Checked;
 	}
-	return ESlateCheckBoxState::Unchecked;
+	return ECheckBoxState::Unchecked;
 }
 
 
diff --git a/Engine/Source/Developer/AutomationWindow/Private/SAutomationExportMenu.h b/Engine/Source/Developer/AutomationWindow/Private/SAutomationExportMenu.h
--- a/Engine/Source/Developer/AutomationWindow/Private/SAutomationExportMenu.h
+++ b/Engine/Source/Developer/AutomationWindow/Private/SAutomationExportMenu.h
@@ -95,7 +95,7 @@ private:
 	 * @param The new state, checked or unchecked.
 	 * @param The type of checkbox set.
 	 */
-	void OnDisplayCheckStateChanged( ESlateCheckBoxState::Type InNewState, EFileExportType::Type CheckType );
+	void OnDisplayCheckStateChanged( ECheckBoxState InNewState, EFileExportType::Type CheckType );
 
 	/** 
 	 * Gets the display state to send to a display filter check box
@@ -103,7 +103,7 @@ private:
 	 * @param The type of checkbox.
 	 * @return The desired checkbox state.
 	 */
-	ESlateCheckBoxState::Type OnGetDisplayCheckState( EFileExportType::Type CheckType ) const;
+	ECheckBoxState OnGetDisplayCheckState( EFileExportType::Type CheckType ) const;
 
 	/**
 	 * Add a notification when the file is exported.
diff --git a/Engine/Source/Developer/AutomationWindow/Private/SAutomationTestItem.cpp b/Engine/Source/Developer/AutomationWindow/Private/SAutomationTestItem.cpp
--- a/Engine/Source/Developer/AutomationWindow/Private/SAutomationTestItem.cpp
+++ b/Engine/Source/Developer/AutomationWindow/Private/SAutomationTestItem.cpp
@@ -361,9 +361,9 @@ FText SAutomationTestItem::GetTestToolTip( int32 ClusterIndex ) const
 }
 
 
-ESlateCheckBoxState::Type SAutomationTestItem::IsTestEnabled() const
+ECheckBoxState SAutomationTestItem::IsTestEnabled() const
 {
-	return TestStatus->IsEnabled() ? ESlateCheckBoxState::Checked : ESlateCheckBoxState::Unchecked;
+	return TestStatus->IsEnabled() ? ECheckBoxState::Checked : ECheckBoxState::Unchecked;
 }
 
 
@@ -550,7 +550,7 @@ const FSlateBrush* SAutomationTestItem::ItemStatus_StatusImage(const int32 Clust
 /* SAutomationTestitem event handlers
  *****************************************************************************/
 
-void SAutomationTestItem::HandleTestingCheckbox_Click(ESlateCheckBoxState::Type)
+void SAutomationTestItem::HandleTestingCheckbox_Click(ECheckBoxState)
 {
 	OnCheckedStateChangedDelegate.ExecuteIfBound(TestStatus);
 }
diff --git a/Engine/Source/Developer/AutomationWindow/Private/SAutomationTestItem.h b/Engine/Source/Developer/AutomationWindow/Private/SAutomationTestItem.h
--- a/Engine/Source/Developer/AutomationWindow/Private/SAutomationTestItem.h
+++ b/Engine/Source/Developer/AutomationWindow/Private/SAutomationTestItem.h
@@ -68,7 +68,7 @@ protected:
 	*
 	* @return true if the test is enabled.
 	*/
-	ESlateCheckBoxState::Type IsTestEnabled( ) const;
+	ECheckBoxState IsTestEnabled( ) const;
 
 	/**
 	 * Returns color that indicates test status per cluster.
@@ -128,7 +128,7 @@ protected:
 private:
 
 	/** Handle the testing checkbox click. */
-	void HandleTestingCheckbox_Click( ESlateCheckBoxState::Type );
+	void HandleTestingCheckbox_Click( ECheckBoxState );
 
 private:
 
diff --git a/Engine/Source/Developer/AutomationWindow/Private/SAutomationWindow.cpp b/Engine/Source/Developer/AutomationWindow/Private/SAutomationWindow.cpp
--- a/Engine/Source/Developer/AutomationWindow/Private/SAutomationWindow.cpp
+++ b/Engine/Source/Developer/AutomationWindow/Private/SAutomationWindow.cpp
@@ -409,19 +409,19 @@ void SAutomationWindow::Construct( const FArguments& InArgs, const IAutomationCo
 END_SLATE_FUNCTION_BUILD_OPTIMIZATION
 
 
-void SAutomationWindow::HandleResultDisplayTypeStateChanged( ESlateCheckBoxState::Type NewRadioState, EAutomationGrapicalDisplayType::Type NewDisplayType)
+void SAutomationWindow::HandleResultDisplayTypeStateChanged( ECheckBoxState NewRadioState, EAutomationGrapicalDisplayType::Type NewDisplayType)
 {
-	if (NewRadioState == ESlateCheckBoxState::Checked)
+	if (NewRadioState == ECheckBoxState::Checked)
 	{
 		GraphicalResultBox->SetDisplayType(NewDisplayType);
 	}
 }
 
-ESlateCheckBoxState::Type SAutomationWindow::HandleResultDisplayTypeIsChecked( EAutomationGrapicalDisplayType::Type InDisplayType ) const
+ECheckBoxState SAutomationWindow::HandleResultDisplayTypeIsChecked( EAutomationGrapicalDisplayType::Type InDisplayType ) const
 {
 	return (GraphicalResultBox->GetDisplayType() == InDisplayType)
-		? ESlateCheckBoxState::Checked
-		: ESlateCheckBoxState::Unchecked;
+		? ECheckBoxState::Checked
+		: ECheckBoxState::Unchecked;
 }
 
 const FSlateBrush* SAutomationWindow::GetTestBackgroundBorderImage() const
@@ -978,12 +978,12 @@ TSharedRef< SWidget > SAutomationWindow::GenerateGroupOptionsMenuContent( )
 }
 
 /** Returns if full size screen shots are enabled */
-ESlateCheckBoxState::Type SAutomationWindow::IsDeviceGroupCheckBoxIsChecked(const int32 DeviceGroupFlag) const
+ECheckBoxState SAutomationWindow::IsDeviceGroupCheckBoxIsChecked(const int32 DeviceGroupFlag) const
 {
-	return AutomationController->IsDeviceGroupFlagSet((EAutomationDeviceGroupTypes::Type)DeviceGroupFlag) ? ESlateCheckBoxState::Checked : ESlateCheckBoxState::Unchecked;
+	return AutomationController->IsDeviceGroupFlagSet((EAutomationDeviceGroupTypes::Type)DeviceGroupFlag) ? ECheckBoxState::Checked : ECheckBoxState::Unchecked;
 }
 /** Toggles if we are collecting full size screenshots */
-void SAutomationWindow::HandleDeviceGroupCheckStateChanged(ESlateCheckBoxState::Type CheckBoxState, const int32 DeviceGroupFlag)
+void SAutomationWindow::HandleDeviceGroupCheckStateChanged(ECheckBoxState CheckBoxState, const int32 DeviceGroupFlag)
 {
 	//Update the device groups
 	AutomationController->ToggleDeviceGroupFlag((EAutomationDeviceGroupTypes::Type)DeviceGroupFlag);
@@ -1143,24 +1143,24 @@ TSharedRef< SWidget > SAutomationWindow::GenerateTestHistoryMenuContent()
 	return MenuBuilder.MakeWidget();
 }
 
-ESlateCheckBoxState::Type SAutomationWindow::IsFullSizeScreenshotsCheckBoxChecked() const
+ECheckBoxState SAutomationWindow::IsFullSizeScreenshotsCheckBoxChecked() const
 {
-	return AutomationController->IsUsingFullSizeScreenshots() ? ESlateCheckBoxState::Checked : ESlateCheckBoxState::Unchecked;
+	return AutomationController->IsUsingFullSizeScreenshots() ? ECheckBoxState::Checked : ECheckBoxState::Unchecked;
 }
 
-void SAutomationWindow::HandleFullSizeScreenshotsBoxCheckStateChanged(ESlateCheckBoxState::Type CheckBoxState)
+void SAutomationWindow::HandleFullSizeScreenshotsBoxCheckStateChanged(ECheckBoxState CheckBoxState)
 {
-	AutomationController->SetUsingFullSizeScreenshots(CheckBoxState == ESlateCheckBoxState::Checked);
+	AutomationController->SetUsingFullSizeScreenshots(CheckBoxState == ECheckBoxState::Checked);
 }
 
-ESlateCheckBoxState::Type SAutomationWindow::IsEnableScreenshotsCheckBoxChecked() const
+ECheckBoxState SAutomationWindow::IsEnableScreenshotsCheckBoxChecked() const
 {
-	return AutomationController->IsScreenshotAllowed() ? ESlateCheckBoxState::Checked : ESlateCheckBoxState::Unchecked;
+	return AutomationController->IsScreenshotAllowed() ? ECheckBoxState::Checked : ECheckBoxState::Unchecked;
 }
 
-void SAutomationWindow::HandleEnableScreenshotsBoxCheckStateChanged(ESlateCheckBoxState::Type CheckBoxState)
+void SAutomationWindow::HandleEnableScreenshotsBoxCheckStateChanged(ECheckBoxState CheckBoxState)
 {
-	AutomationController->SetScreenshotsEnabled(CheckBoxState == ESlateCheckBoxState::Checked);
+	AutomationController->SetScreenshotsEnabled(CheckBoxState == ECheckBoxState::Checked);
 }
 
 bool SAutomationWindow::IsFullSizeScreenshotsOptionEnabled() const
@@ -1282,9 +1282,9 @@ EVisibility SAutomationWindow::GetTestGraphVisibility( ) const
 }
 
 
-void SAutomationWindow::HeaderCheckboxStateChange(ESlateCheckBoxState::Type InCheckboxState)
+void SAutomationWindow::HeaderCheckboxStateChange(ECheckBoxState InCheckboxState)
 {
-	const bool bState = (InCheckboxState == ESlateCheckBoxState::Checked)? true : false;
+	const bool bState = (InCheckboxState == ECheckBoxState::Checked)? true : false;
 
 	AutomationController->SetVisibleTestsEnabled(bState);
 }
@@ -1641,14 +1641,14 @@ void SAutomationWindow::OnToggleErrorFilter()
 }
 
 
-ESlateCheckBoxState::Type SAutomationWindow::IsTrackingHistory() const
+ECheckBoxState SAutomationWindow::IsTrackingHistory() const
 {
-	return bIsTrackingHistory ? ESlateCheckBoxState::Checked : ESlateCheckBoxState::Unchecked;
+	return bIsTrackingHistory ? ECheckBoxState::Checked : ECheckBoxState::Unchecked;
 }
 
-void SAutomationWindow::OnToggleTrackHistory(ESlateCheckBoxState::Type InState)
+void SAutomationWindow::OnToggleTrackHistory(ECheckBoxState InState)
 {
-	bIsTrackingHistory = InState == ESlateCheckBoxState::Checked;
+	bIsTrackingHistory = InState == ECheckBoxState::Checked;
 
 	AutomationController->TrackReportHistory(bIsTrackingHistory, NumHistoryElementsToTrack);
 
diff --git a/Engine/Source/Developer/AutomationWindow/Private/SAutomationWindow.h b/Engine/Source/Developer/AutomationWindow/Private/SAutomationWindow.h
--- a/Engine/Source/Developer/AutomationWindow/Private/SAutomationWindow.h
+++ b/Engine/Source/Developer/AutomationWindow/Private/SAutomationWindow.h
@@ -195,7 +195,7 @@ private:
 	void OnTestSelectionChanged(TSharedPtr<IAutomationReport> Selection, ESelectInfo::Type SelectInfo);
 	
 	/** Called when the header checkbox's state is changed. Will go through all the Unit Tests available and change the state of their checkbox. */
-	void HeaderCheckboxStateChange(ESlateCheckBoxState::Type InCheckboxState);	/** Creates table row for each visible automation test */
+	void HeaderCheckboxStateChange(ECheckBoxState InCheckboxState);	/** Creates table row for each visible automation test */
 	
 	/** Rebuilds the platform icon header */
 	void RebuildPlatformIcons();
@@ -265,29 +265,29 @@ private:
 	void OnToggleErrorFilter();
 	
 	/** Returns if we're tracking history for automation tests */
-	ESlateCheckBoxState::Type IsTrackingHistory() const;
+	ECheckBoxState IsTrackingHistory() const;
 	
 	/** Toggles whether we are tracking history of automation tests */
-	void OnToggleTrackHistory(ESlateCheckBoxState::Type InState);
+	void OnToggleTrackHistory(ECheckBoxState InState);
 
 	/** Returns if full size screen shots are enabled */
-	ESlateCheckBoxState::Type IsFullSizeScreenshotsCheckBoxChecked() const;
+	ECheckBoxState IsFullSizeScreenshotsCheckBoxChecked() const;
 	/** Toggles if we are collecting full size screenshots */
-	void HandleFullSizeScreenshotsBoxCheckStateChanged(ESlateCheckBoxState::Type CheckBoxState);
+	void HandleFullSizeScreenshotsBoxCheckStateChanged(ECheckBoxState CheckBoxState);
 
 	/** Returns if screen shots are enabled */
-	ESlateCheckBoxState::Type IsEnableScreenshotsCheckBoxChecked() const;
+	ECheckBoxState IsEnableScreenshotsCheckBoxChecked() const;
 	/** Toggles if we are taking screenshots */
-	void HandleEnableScreenshotsBoxCheckStateChanged(ESlateCheckBoxState::Type CheckBoxState);
+	void HandleEnableScreenshotsBoxCheckStateChanged(ECheckBoxState CheckBoxState);
 
 	/** Returns if the full size screenshots option is enabled */
 	bool IsFullSizeScreenshotsOptionEnabled() const;
 
 	/** Returns if a device group is enabled */
-	ESlateCheckBoxState::Type IsDeviceGroupCheckBoxIsChecked(const int32 DeviceGroupFlag) const;
+	ECheckBoxState IsDeviceGroupCheckBoxIsChecked(const int32 DeviceGroupFlag) const;
 	
 	/** Toggles a device group flag */
-	void HandleDeviceGroupCheckStateChanged(ESlateCheckBoxState::Type CheckBoxState, const int32 DeviceGroupFlag);
+	void HandleDeviceGroupCheckStateChanged(ECheckBoxState CheckBoxState, const int32 DeviceGroupFlag);
 	
 	/** Sets the number of times to repeat the tests */
 	void OnChangeRepeatCount(int32 InNewValue);
@@ -460,10 +460,10 @@ private:
 	EVisibility GetTestGraphVisibility( ) const;
 
 	/** Handles changing the display type for the graphical results window. */
-	void HandleResultDisplayTypeStateChanged( ESlateCheckBoxState::Type NewRadioState, EAutomationGrapicalDisplayType::Type NewDisplayType );
+	void HandleResultDisplayTypeStateChanged( ECheckBoxState NewRadioState, EAutomationGrapicalDisplayType::Type NewDisplayType );
 
 	/** Handles checking if a display type is active for the graphical results window. */
-	ESlateCheckBoxState::Type HandleResultDisplayTypeIsChecked( EAutomationGrapicalDisplayType::Type InDisplayType ) const;
+	ECheckBoxState HandleResultDisplayTypeIsChecked( EAutomationGrapicalDisplayType::Type InDisplayType ) const;
 
 	/** 
 	 * Gets the visibility for the throbber.
diff --git a/Engine/Source/Developer/CollisionAnalyzer/Private/SCAQueryDetails.cpp b/Engine/Source/Developer/CollisionAnalyzer/Private/SCAQueryDetails.cpp
--- a/Engine/Source/Developer/CollisionAnalyzer/Private/SCAQueryDetails.cpp
+++ b/Engine/Source/Developer/CollisionAnalyzer/Private/SCAQueryDetails.cpp
@@ -268,15 +268,15 @@ void SCAQueryDetails::ResultListSelectionChanged(TSharedPtr<FCAHitInfo> Selected
 }
 
 
-void SCAQueryDetails::OnToggleShowMisses(ESlateCheckBoxState::Type InCheckboxState)
+void SCAQueryDetails::OnToggleShowMisses(ECheckBoxState InCheckboxState)
 {
-	bShowMisses = (InCheckboxState == ESlateCheckBoxState::Checked);
+	bShowMisses = (InCheckboxState == ECheckBoxState::Checked);
 	UpdateResultList();
 }
 
-ESlateCheckBoxState::Type SCAQueryDetails::GetShowMissesState() const
+ECheckBoxState SCAQueryDetails::GetShowMissesState() const
 {
-	return bShowMisses ? ESlateCheckBoxState::Checked : ESlateCheckBoxState::Unchecked;
+	return bShowMisses ? ECheckBoxState::Checked : ECheckBoxState::Unchecked;
 }
 
 
diff --git a/Engine/Source/Developer/CollisionAnalyzer/Private/SCAQueryDetails.h b/Engine/Source/Developer/CollisionAnalyzer/Private/SCAQueryDetails.h
--- a/Engine/Source/Developer/CollisionAnalyzer/Private/SCAQueryDetails.h
+++ b/Engine/Source/Developer/CollisionAnalyzer/Private/SCAQueryDetails.h
@@ -46,12 +46,12 @@ public:
 	// Get delegates
 	FText GetStartText() const;
 	FText GetEndText() const;
-	ESlateCheckBoxState::Type GetShowMissesState() const;
+	ECheckBoxState GetShowMissesState() const;
 	// List delegates
 	TSharedRef<ITableRow> ResultListGenerateRow(TSharedPtr<FCAHitInfo> Info, const TSharedRef<STableViewBase>& OwnerTable);
 	void ResultListSelectionChanged(TSharedPtr<FCAHitInfo> SelectedInfos, ESelectInfo::Type SelectInfo);
 	// Handler delegates
-	void OnToggleShowMisses(ESlateCheckBoxState::Type InCheckboxState);
+	void OnToggleShowMisses(ECheckBoxState InCheckboxState);
 
 private:
 
diff --git a/Engine/Source/Developer/CollisionAnalyzer/Private/SCollisionAnalyzer.cpp b/Engine/Source/Developer/CollisionAnalyzer/Private/SCollisionAnalyzer.cpp
--- a/Engine/Source/Developer/CollisionAnalyzer/Private/SCollisionAnalyzer.cpp
+++ b/Engine/Source/Developer/CollisionAnalyzer/Private/SCollisionAnalyzer.cpp
@@ -665,9 +665,9 @@ FString SCollisionAnalyzer::GetStatusText() const
 	return FString::Printf(TEXT("Total: %d queries over %d frames. Shown: %d queries"), Analyzer->Queries.Num(), Analyzer->GetNumFramesOfRecording(), TotalNumQueries);
 }
 
-ESlateCheckBoxState::Type SCollisionAnalyzer::GetDrawRecentState() const
+ECheckBoxState SCollisionAnalyzer::GetDrawRecentState() const
 {
-	return bDrawRecentQueries ? ESlateCheckBoxState::Checked : ESlateCheckBoxState::Unchecked;
+	return bDrawRecentQueries ? ECheckBoxState::Checked : ECheckBoxState::Unchecked;
 }
 
 
@@ -679,48 +679,48 @@ FReply SCollisionAnalyzer::OnRecordButtonClicked()
 	return FReply::Handled();
 }
 
-void SCollisionAnalyzer::OnDrawRecentChanged(ESlateCheckBoxState::Type NewState)
+void SCollisionAnalyzer::OnDrawRecentChanged(ECheckBoxState NewState)
 {
-	bDrawRecentQueries = (NewState == ESlateCheckBoxState::Checked);
+	bDrawRecentQueries = (NewState == ECheckBoxState::Checked);
 }
 
 // By frame
 
-ESlateCheckBoxState::Type SCollisionAnalyzer::GetGroupByFrameState() const
+ECheckBoxState SCollisionAnalyzer::GetGroupByFrameState() const
 {
-	return (GroupBy == EQueryGroupMode::ByFrameNum) ? ESlateCheckBoxState::Checked : ESlateCheckBoxState::Unchecked;
+	return (GroupBy == EQueryGroupMode::ByFrameNum) ? ECheckBoxState::Checked : ECheckBoxState::Unchecked;
 }
 
-void SCollisionAnalyzer::OnGroupByFrameChanged(ESlateCheckBoxState::Type NewState)
+void SCollisionAnalyzer::OnGroupByFrameChanged(ECheckBoxState NewState)
 {
-	GroupBy = (NewState == ESlateCheckBoxState::Checked) ? EQueryGroupMode::ByFrameNum : EQueryGroupMode::Ungrouped;
+	GroupBy = (NewState == ECheckBoxState::Checked) ? EQueryGroupMode::ByFrameNum : EQueryGroupMode::Ungrouped;
 	RebuildFilteredList();
 }
 
 
 // By Tag
 
-ESlateCheckBoxState::Type SCollisionAnalyzer::GetGroupByTagState() const
+ECheckBoxState SCollisionAnalyzer::GetGroupByTagState() const
 {
-	return (GroupBy == EQueryGroupMode::ByTag) ? ESlateCheckBoxState::Checked : ESlateCheckBoxState::Unchecked;
+	return (GroupBy == EQueryGroupMode::ByTag) ? ECheckBoxState::Checked : ECheckBoxState::Unchecked;
 }
 
-void SCollisionAnalyzer::OnGroupByTagChanged(ESlateCheckBoxState::Type NewState)
+void SCollisionAnalyzer::OnGroupByTagChanged(ECheckBoxState NewState)
 {
-	GroupBy = (NewState == ESlateCheckBoxState::Checked) ? EQueryGroupMode::ByTag : EQueryGroupMode::Ungrouped;
+	GroupBy = (NewState == ECheckBoxState::Checked) ? EQueryGroupMode::ByTag : EQueryGroupMode::Ungrouped;
 	RebuildFilteredList();
 }
 
 // By Owner
 
-ESlateCheckBoxState::Type SCollisionAnalyzer::GetGroupByOwnerState() const
+ECheckBoxState SCollisionAnalyzer::GetGroupByOwnerState() const
 {
-	return (GroupBy == EQueryGroupMode::ByOwnerTag) ? ESlateCheckBoxState::Checked : ESlateCheckBoxState::Unchecked;
+	return (GroupBy == EQueryGroupMode::ByOwnerTag) ? ECheckBoxState::Checked : ECheckBoxState::Unchecked;
 }
 
-void SCollisionAnalyzer::OnGroupByOwnerChanged(ESlateCheckBoxState::Type NewState)
+void SCollisionAnalyzer::OnGroupByOwnerChanged(ECheckBoxState NewState)
 {
-	GroupBy = (NewState == ESlateCheckBoxState::Checked) ? EQueryGroupMode::ByOwnerTag : EQueryGroupMode::Ungrouped;
+	GroupBy = (NewState == ECheckBoxState::Checked) ? EQueryGroupMode::ByOwnerTag : EQueryGroupMode::Ungrouped;
 	RebuildFilteredList();
 }
 
diff --git a/Engine/Source/Developer/CollisionAnalyzer/Private/SCollisionAnalyzer.h b/Engine/Source/Developer/CollisionAnalyzer/Private/SCollisionAnalyzer.h
--- a/Engine/Source/Developer/CollisionAnalyzer/Private/SCollisionAnalyzer.h
+++ b/Engine/Source/Developer/CollisionAnalyzer/Private/SCollisionAnalyzer.h
@@ -95,18 +95,18 @@ public:
 	// Get delegates
 	const FSlateBrush* GetRecordButtonBrush() const;
 	FString GetStatusText() const;
-	ESlateCheckBoxState::Type GetDrawRecentState() const;
-	ESlateCheckBoxState::Type GetGroupByFrameState() const;
-	ESlateCheckBoxState::Type GetGroupByTagState() const;
-	ESlateCheckBoxState::Type GetGroupByOwnerState() const;
+	ECheckBoxState GetDrawRecentState() const;
+	ECheckBoxState GetGroupByFrameState() const;
+	ECheckBoxState GetGroupByTagState() const;
+	ECheckBoxState GetGroupByOwnerState() const;
 	EColumnSortMode::Type GetIDSortMode() const;
 	EColumnSortMode::Type GetTimeSortMode() const;
 	// Handler delegates
 	FReply OnRecordButtonClicked();
-	void OnDrawRecentChanged(ESlateCheckBoxState::Type NewState);
-	void OnGroupByFrameChanged(ESlateCheckBoxState::Type NewState);
-	void OnGroupByTagChanged(ESlateCheckBoxState::Type NewState);
-	void OnGroupByOwnerChanged(ESlateCheckBoxState::Type NewState);
+	void OnDrawRecentChanged(ECheckBoxState NewState);
+	void OnGroupByFrameChanged(ECheckBoxState NewState);
+	void OnGroupByTagChanged(ECheckBoxState NewState);
+	void OnGroupByOwnerChanged(ECheckBoxState NewState);
 	void FilterTextCommitted(const FText& CommentText, ETextCommit::Type CommitInfo);
 	void OnSortByChanged(const EColumnSortPriority::Type SortPriority, const FName& ColumnName, const EColumnSortMode::Type NewSortMode);
 	// Table delegates
diff --git a/Engine/Source/Developer/DeviceManager/Private/Widgets/Browser/SDeviceBrowserDeviceListRow.h b/Engine/Source/Developer/DeviceManager/Private/Widgets/Browser/SDeviceBrowserDeviceListRow.h
--- a/Engine/Source/Developer/DeviceManager/Private/Widgets/Browser/SDeviceBrowserDeviceListRow.h
+++ b/Engine/Source/Developer/DeviceManager/Private/Widgets/Browser/SDeviceBrowserDeviceListRow.h
@@ -146,20 +146,20 @@ private:
 	}
 
 	// Callback for changing this row's Share check box state.
-	void HandleShareCheckBoxStateChanged( ESlateCheckBoxState::Type NewState )
+	void HandleShareCheckBoxStateChanged( ECheckBoxState NewState )
 	{
-		DeviceService->SetShared(NewState == ESlateCheckBoxState::Checked);
+		DeviceService->SetShared(NewState == ECheckBoxState::Checked);
 	}
 
 	// Callback for getting the state of the 'Share' check box.
-	ESlateCheckBoxState::Type HandleShareCheckBoxIsChecked( ) const
+	ECheckBoxState HandleShareCheckBoxIsChecked( ) const
 	{
 		if (DeviceService->IsShared())
 		{
-			return ESlateCheckBoxState::Checked;
+			return ECheckBoxState::Checked;
 		}
 
-		return ESlateCheckBoxState::Unchecked;
+		return ECheckBoxState::Unchecked;
 	}
 
 	// Callback for getting the enabled state of the 'Share' check box.
diff --git a/Engine/Source/Developer/DeviceManager/Private/Widgets/Browser/SDeviceBrowserFilterBar.cpp b/Engine/Source/Developer/DeviceManager/Private/Widgets/Browser/SDeviceBrowserFilterBar.cpp
--- a/Engine/Source/Developer/DeviceManager/Private/Widgets/Browser/SDeviceBrowserFilterBar.cpp
+++ b/Engine/Source/Developer/DeviceManager/Private/Widgets/Browser/SDeviceBrowserFilterBar.cpp
@@ -89,20 +89,20 @@ void SDeviceBrowserFilterBar::HandleFilterStringTextChanged( const FText& NewTex
 }
 
 
-void SDeviceBrowserFilterBar::HandlePlatformListRowCheckStateChanged(ESlateCheckBoxState::Type CheckState, TSharedPtr<FDeviceBrowserFilterEntry> PlatformEntry)
+void SDeviceBrowserFilterBar::HandlePlatformListRowCheckStateChanged(ECheckBoxState CheckState, TSharedPtr<FDeviceBrowserFilterEntry> PlatformEntry)
 {
-	Filter->SetPlatformEnabled(PlatformEntry->PlatformName, CheckState == ESlateCheckBoxState::Checked);
+	Filter->SetPlatformEnabled(PlatformEntry->PlatformName, CheckState == ECheckBoxState::Checked);
 }
 
 
-ESlateCheckBoxState::Type SDeviceBrowserFilterBar::HandlePlatformListRowIsChecked(TSharedPtr<FDeviceBrowserFilterEntry> PlatformEntry) const
+ECheckBoxState SDeviceBrowserFilterBar::HandlePlatformListRowIsChecked(TSharedPtr<FDeviceBrowserFilterEntry> PlatformEntry) const
 {
 	if (Filter->IsPlatformEnabled(PlatformEntry->PlatformName))
 	{
-		return ESlateCheckBoxState::Checked;
+		return ECheckBoxState::Checked;
 	}
 
-	return ESlateCheckBoxState::Unchecked;
+	return ECheckBoxState::Unchecked;
 }
 
 
diff --git a/Engine/Source/Developer/DeviceManager/Private/Widgets/Browser/SDeviceBrowserFilterBar.h b/Engine/Source/Developer/DeviceManager/Private/Widgets/Browser/SDeviceBrowserFilterBar.h
--- a/Engine/Source/Developer/DeviceManager/Private/Widgets/Browser/SDeviceBrowserFilterBar.h
+++ b/Engine/Source/Developer/DeviceManager/Private/Widgets/Browser/SDeviceBrowserFilterBar.h
@@ -47,10 +47,10 @@ private:
 	void HandleFilterStringTextChanged(const FText& NewText);
 
 	// Callback for changing the checked state of the given platform filter row.
-	void HandlePlatformListRowCheckStateChanged(ESlateCheckBoxState::Type CheckState, TSharedPtr<FDeviceBrowserFilterEntry> PlatformEntry);
+	void HandlePlatformListRowCheckStateChanged(ECheckBoxState CheckState, TSharedPtr<FDeviceBrowserFilterEntry> PlatformEntry);
 
 	// Callback for getting the checked state of the given platform filter row.
-	ESlateCheckBoxState::Type HandlePlatformListRowIsChecked(TSharedPtr<FDeviceBrowserFilterEntry> PlatformEntry) const;
+	ECheckBoxState HandlePlatformListRowIsChecked(TSharedPtr<FDeviceBrowserFilterEntry> PlatformEntry) const;
 
 	// Callback for getting the text for a row in the platform filter drop-down.
 	FString HandlePlatformListRowText(TSharedPtr<FDeviceBrowserFilterEntry> PlatformEntry) const;
diff --git a/Engine/Source/Developer/DeviceManager/Private/Widgets/Processes/SDeviceProcesses.cpp b/Engine/Source/Developer/DeviceManager/Private/Widgets/Processes/SDeviceProcesses.cpp
--- a/Engine/Source/Developer/DeviceManager/Private/Widgets/Processes/SDeviceProcesses.cpp
+++ b/Engine/Source/Developer/DeviceManager/Private/Widgets/Processes/SDeviceProcesses.cpp
@@ -283,22 +283,22 @@ void SDeviceProcesses::HandleModelSelectedDeviceServiceChanged( )
 }
 
 
-void SDeviceProcesses::HandleProcessTreeCheckBoxCheckStateChanged( ESlateCheckBoxState::Type NewState )
+void SDeviceProcesses::HandleProcessTreeCheckBoxCheckStateChanged( ECheckBoxState NewState )
 {
-	ShowProcessTree = (NewState == ESlateCheckBoxState::Checked);
+	ShowProcessTree = (NewState == ECheckBoxState::Checked);
 
 	ReloadProcessList(false);
 }
 
 
-ESlateCheckBoxState::Type SDeviceProcesses::HandleProcessTreeCheckBoxIsChecked( ) const
+ECheckBoxState SDeviceProcesses::HandleProcessTreeCheckBoxIsChecked( ) const
 {
 	if (ShowProcessTree)
 	{
-		return ESlateCheckBoxState::Checked;
+		return ECheckBoxState::Checked;
 	}
 
-	return ESlateCheckBoxState::Unchecked;
+	return ECheckBoxState::Unchecked;
 }
 
 
diff --git a/Engine/Source/Developer/DeviceManager/Private/Widgets/Processes/SDeviceProcesses.h b/Engine/Source/Developer/DeviceManager/Private/Widgets/Processes/SDeviceProcesses.h
--- a/Engine/Source/Developer/DeviceManager/Private/Widgets/Processes/SDeviceProcesses.h
+++ b/Engine/Source/Developer/DeviceManager/Private/Widgets/Processes/SDeviceProcesses.h
@@ -56,10 +56,10 @@ private:
 	void HandleModelSelectedDeviceServiceChanged( );
 
 	// Callback for changing 'Show process tree' check box state.
-	void HandleProcessTreeCheckBoxCheckStateChanged( ESlateCheckBoxState::Type NewState );
+	void HandleProcessTreeCheckBoxCheckStateChanged( ECheckBoxState NewState );
 
 	// Callback for determining checked state of the 'Show process tree' check box.
-	ESlateCheckBoxState::Type HandleProcessTreeCheckBoxIsChecked( ) const;
+	ECheckBoxState HandleProcessTreeCheckBoxIsChecked( ) const;
 
 	// Callback for generating a row widget in the process tree view.
 	TSharedRef<ITableRow> HandleProcessTreeViewGenerateRow( FDeviceProcessesProcessTreeNodePtr Item, const TSharedRef<STableViewBase>& OwnerTable );
diff --git a/Engine/Source/Developer/FriendsAndChat/Private/Models/FriendsUserSettingsViewModel.cpp b/Engine/Source/Developer/FriendsAndChat/Private/Models/FriendsUserSettingsViewModel.cpp
--- a/Engine/Source/Developer/FriendsAndChat/Private/Models/FriendsUserSettingsViewModel.cpp
+++ b/Engine/Source/Developer/FriendsAndChat/Private/Models/FriendsUserSettingsViewModel.cpp
@@ -19,13 +19,13 @@ private:
 	{
 	}
 
-	virtual void HandleCheckboxStateChanged(ESlateCheckBoxState::Type NewState, EUserSettngsType::Type OptionType) override
+	virtual void HandleCheckboxStateChanged(ECheckBoxState NewState, EUserSettngsType::Type OptionType) override
 	{
 		switch(OptionType)
 		{
 			case EUserSettngsType::ShowNotifications:
 			{
-				bShowNotifications = NewState == ESlateCheckBoxState::Checked ? true : false;
+				bShowNotifications = NewState == ECheckBoxState::Checked ? true : false;
 			}
 			break;
 		};
@@ -36,16 +36,16 @@ private:
 		UserSettings.Add(EUserSettngsType::ShowNotifications);
 	}
 
-	virtual ESlateCheckBoxState::Type GetOptionCheckState(EUserSettngsType::Type Option) const override
+	virtual ECheckBoxState GetOptionCheckState(EUserSettngsType::Type Option) const override
 	{
 		switch(Option)
 		{
 			case EUserSettngsType::ShowNotifications:
 			{
-				return bShowNotifications ? ESlateCheckBoxState::Checked : ESlateCheckBoxState::Unchecked; break;
+				return bShowNotifications ? ECheckBoxState::Checked : ECheckBoxState::Unchecked; break;
 			}
 			default:
-			return ESlateCheckBoxState::Undetermined;
+			return ECheckBoxState::Undetermined;
 		};
 	}
 
diff --git a/Engine/Source/Developer/FriendsAndChat/Private/Models/FriendsUserSettingsViewModel.h b/Engine/Source/Developer/FriendsAndChat/Private/Models/FriendsUserSettingsViewModel.h
--- a/Engine/Source/Developer/FriendsAndChat/Private/Models/FriendsUserSettingsViewModel.h
+++ b/Engine/Source/Developer/FriendsAndChat/Private/Models/FriendsUserSettingsViewModel.h
@@ -7,9 +7,9 @@ class FFriendsUserSettingsViewModel
 {
 public:
 	virtual ~FFriendsUserSettingsViewModel() {}
-	virtual void HandleCheckboxStateChanged(ESlateCheckBoxState::Type NewState, EUserSettngsType::Type OptionType) = 0;
+	virtual void HandleCheckboxStateChanged(ECheckBoxState NewState, EUserSettngsType::Type OptionType) = 0;
 	virtual void EnumerateUserSettings(TArray<EUserSettngsType::Type>& UserSettings) = 0;
-	virtual ESlateCheckBoxState::Type GetOptionCheckState(EUserSettngsType::Type Option) const = 0;
+	virtual ECheckBoxState GetOptionCheckState(EUserSettngsType::Type Option) const = 0;
 };
 
 /**
diff --git a/Engine/Source/Developer/FriendsAndChat/Private/UI/Widgets/SChatWindow.cpp b/Engine/Source/Developer/FriendsAndChat/Private/UI/Widgets/SChatWindow.cpp
--- a/Engine/Source/Developer/FriendsAndChat/Private/UI/Widgets/SChatWindow.cpp
+++ b/Engine/Source/Developer/FriendsAndChat/Private/UI/Widgets/SChatWindow.cpp
@@ -584,17 +584,17 @@ private:
 		}
 	}
 
-	void OnGlobalOptionChanged(ESlateCheckBoxState::Type NewState)
+	void OnGlobalOptionChanged(ECheckBoxState NewState)
 	{
-		const bool bDisabled = NewState == ESlateCheckBoxState::Unchecked;
+		const bool bDisabled = NewState == ECheckBoxState::Unchecked;
 		ViewModel->SetAllowGlobalChat(!bDisabled);
 
 		FFriendsAndChatManager::Get()->GetAnalytics().RecordToggleChat(TEXT("Global"), !bDisabled, TEXT("Social.Chat.Toggle"));
 	}
 
-	ESlateCheckBoxState::Type GetGlobalOptionState() const
+	ECheckBoxState GetGlobalOptionState() const
 	{
-		return ViewModel->IsGlobalChatEnabled() ? ESlateCheckBoxState::Checked : ESlateCheckBoxState::Unchecked;
+		return ViewModel->IsGlobalChatEnabled() ? ECheckBoxState::Checked : ECheckBoxState::Unchecked;
 	}
 
 	virtual bool SupportsKeyboardFocus() const override
diff --git a/Engine/Source/Developer/InternationalizationSettings/Private/InternationalizationSettingsModelDetails.cpp b/Engine/Source/Developer/InternationalizationSettings/Private/InternationalizationSettingsModelDetails.cpp
--- a/Engine/Source/Developer/InternationalizationSettings/Private/InternationalizationSettingsModelDetails.cpp
+++ b/Engine/Source/Developer/InternationalizationSettings/Private/InternationalizationSettingsModelDetails.cpp
@@ -215,7 +215,7 @@ void FInternationalizationSettingsModelDetails::CustomizeDetails( IDetailLayoutB
 	.MaxDesiredWidth(300.0f)
 	[
 		SAssignNew(FieldNamesCheckBox, SCheckBox)
-		.IsChecked(Model->ShouldLoadLocalizedPropertyNames() ? ESlateCheckBoxState::Checked : ESlateCheckBoxState::Unchecked)
+		.IsChecked(Model->ShouldLoadLocalizedPropertyNames() ? ECheckBoxState::Checked : ECheckBoxState::Unchecked)
 		.ToolTipText(FieldNamesToolTipText)
 		.OnCheckStateChanged(this, &FInternationalizationSettingsModelDetails::ShoudLoadLocalizedFieldNamesCheckChanged)
 	];
@@ -408,7 +408,7 @@ EVisibility FInternationalizationSettingsModelDetails::GetInternationalizationRe
 	return RequiresRestart ? EVisibility::Visible : EVisibility::Collapsed;
 }
 
-void FInternationalizationSettingsModelDetails::ShoudLoadLocalizedFieldNamesCheckChanged(ESlateCheckBoxState::Type CheckState)
+void FInternationalizationSettingsModelDetails::ShoudLoadLocalizedFieldNamesCheckChanged(ECheckBoxState CheckState)
 {
 	HandleShutdownPostPackagesSaved();
 }
diff --git a/Engine/Source/Developer/InternationalizationSettings/Public/InternationalizationSettingsModelDetails.h b/Engine/Source/Developer/InternationalizationSettings/Public/InternationalizationSettingsModelDetails.h
--- a/Engine/Source/Developer/InternationalizationSettings/Public/InternationalizationSettingsModelDetails.h
+++ b/Engine/Source/Developer/InternationalizationSettings/Public/InternationalizationSettingsModelDetails.h
@@ -63,7 +63,7 @@ private:
 	EVisibility GetInternationalizationRestartRowVisibility() const;
 
 	/** Delegate called when the the checked state of whether or not field names should be localized has changed. */
-	void ShoudLoadLocalizedFieldNamesCheckChanged(ESlateCheckBoxState::Type CheckState);
+	void ShoudLoadLocalizedFieldNamesCheckChanged(ECheckBoxState CheckState);
 
 	/** Write to config now the Editor is shutting down (all packages are saved) */
 	void HandleShutdownPostPackagesSaved();
diff --git a/Engine/Source/Developer/LogVisualizer/Private/SFilterList.cpp b/Engine/Source/Developer/LogVisualizer/Private/SFilterList.cpp
--- a/Engine/Source/Developer/LogVisualizer/Private/SFilterList.cpp
+++ b/Engine/Source/Developer/LogVisualizer/Private/SFilterList.cpp
@@ -200,9 +200,9 @@ public:
 
 private:
 	/** Handler for when the filter checkbox is clicked */
-	void FilterToggled(ESlateCheckBoxState::Type NewState)
+	void FilterToggled(ECheckBoxState NewState)
 	{
-		bEnabled = NewState == ESlateCheckBoxState::Checked;
+		bEnabled = NewState == ECheckBoxState::Checked;
 		OnFilterChanged.ExecuteIfBound();
 	}
 
@@ -284,15 +284,15 @@ private:
 	}
 
 	/** Handler to determine the "checked" state of the filter checkbox */
-	ESlateCheckBoxState::Type IsChecked() const
+	ECheckBoxState IsChecked() const
 	{
-		return bEnabled ? ESlateCheckBoxState::Checked : ESlateCheckBoxState::Unchecked;
+		return bEnabled ? ECheckBoxState::Checked : ECheckBoxState::Unchecked;
 	}
 
 	/** Handler to determine the color of the checkbox when it is checked */
 	FSlateColor GetFilterForegroundColor() const
 	{
-		return IsChecked() ? FilterColor : FLinearColor::White;
+		return IsChecked() == ECheckBoxState::Checked ? FilterColor : FLinearColor::White;
 	}
 
 	/** Handler to determine the padding of the checkbox text when it is pressed */
@@ -305,7 +305,7 @@ private:
 	FSlateColor GetFilterNameColorAndOpacity() const
 	{
 		const float DimFactor = 0.75f;	
-		return IsChecked() ? (IsHovered() ? ColorCategory * DimFactor : ColorCategory) : (IsHovered() ? FLinearColor::White : FLinearColor::White * DimFactor);
+		return IsChecked() == ECheckBoxState::Checked ? ( IsHovered() ? ColorCategory * DimFactor : ColorCategory ) : ( IsHovered() ? FLinearColor::White : FLinearColor::White * DimFactor );
 	}
 
 private:
diff --git a/Engine/Source/Developer/LogVisualizer/Private/SLogVisualizer.cpp b/Engine/Source/Developer/LogVisualizer/Private/SLogVisualizer.cpp
--- a/Engine/Source/Developer/LogVisualizer/Private/SLogVisualizer.cpp
+++ b/Engine/Source/Developer/LogVisualizer/Private/SLogVisualizer.cpp
@@ -2099,51 +2099,51 @@ void SLogVisualizer::OnSetHistogramWindowValue(float NewValue)
 	InvalidateCanvas();
 }
 
-void SLogVisualizer::OnDrawLogEntriesPathChanged(ESlateCheckBoxState::Type NewState)
+void SLogVisualizer::OnDrawLogEntriesPathChanged(ECheckBoxState NewState)
 {
-	bDrawLogEntriesPath = (NewState == ESlateCheckBoxState::Checked);
+	bDrawLogEntriesPath = (NewState == ECheckBoxState::Checked);
 	InvalidateCanvas();
 }
 
-ESlateCheckBoxState::Type SLogVisualizer::GetDrawLogEntriesPathState() const
+ECheckBoxState SLogVisualizer::GetDrawLogEntriesPathState() const
 {
-	return bDrawLogEntriesPath ? ESlateCheckBoxState::Checked : ESlateCheckBoxState::Unchecked;
+	return bDrawLogEntriesPath ? ECheckBoxState::Checked : ECheckBoxState::Unchecked;
 }
 
-void SLogVisualizer::OnIgnoreTrivialLogs(ESlateCheckBoxState::Type NewState)
+void SLogVisualizer::OnIgnoreTrivialLogs(ECheckBoxState NewState)
 {
-	bIgnoreTrivialLogs = (NewState == ESlateCheckBoxState::Checked);
+	bIgnoreTrivialLogs = (NewState == ECheckBoxState::Checked);
 	DoFullUpdate();
 }
 
-ESlateCheckBoxState::Type SLogVisualizer::GetIgnoreTrivialLogs() const
+ECheckBoxState SLogVisualizer::GetIgnoreTrivialLogs() const
 {
-	return bIgnoreTrivialLogs ? ESlateCheckBoxState::Checked : ESlateCheckBoxState::Unchecked;
+	return bIgnoreTrivialLogs ? ECheckBoxState::Checked : ECheckBoxState::Unchecked;
 }
 
-void SLogVisualizer::OnChangeHistogramLabelLocation(ESlateCheckBoxState::Type NewState)
+void SLogVisualizer::OnChangeHistogramLabelLocation(ECheckBoxState NewState)
 {
-	bShowHistogramLabelsOutside = (NewState == ESlateCheckBoxState::Checked);
+	bShowHistogramLabelsOutside = (NewState == ECheckBoxState::Checked);
 	InvalidateCanvas();
 }
 
-ESlateCheckBoxState::Type SLogVisualizer::GetHistogramLabelLocation() const
+ECheckBoxState SLogVisualizer::GetHistogramLabelLocation() const
 {
-	return bShowHistogramLabelsOutside ? ESlateCheckBoxState::Checked : ESlateCheckBoxState::Unchecked;
+	return bShowHistogramLabelsOutside ? ECheckBoxState::Checked : ECheckBoxState::Unchecked;
 }
 
-void SLogVisualizer::OnStickToLastData(ESlateCheckBoxState::Type NewState)
+void SLogVisualizer::OnStickToLastData(ECheckBoxState NewState)
 {
-	bStickToLastData = (NewState == ESlateCheckBoxState::Checked);
+	bStickToLastData = (NewState == ECheckBoxState::Checked);
 	InvalidateCanvas();
 }
 
-ESlateCheckBoxState::Type SLogVisualizer::GetStickToLastData() const
+ECheckBoxState SLogVisualizer::GetStickToLastData() const
 {
-	return bStickToLastData ? ESlateCheckBoxState::Checked : ESlateCheckBoxState::Unchecked;
+	return bStickToLastData ? ECheckBoxState::Checked : ECheckBoxState::Unchecked;
 }
 
-void SLogVisualizer::OnToggleCamera(ESlateCheckBoxState::Type NewState)
+void SLogVisualizer::OnToggleCamera(ECheckBoxState NewState)
 {
 	UWorld* World = GetWorld();
 	if (ALogVisualizerCameraController::IsEnabled(World))
@@ -2157,26 +2157,26 @@ void SLogVisualizer::OnToggleCamera(ESlateCheckBoxState::Type NewState)
 	InvalidateCanvas();
 }
 
-ESlateCheckBoxState::Type SLogVisualizer::GetToggleCameraState() const
+ECheckBoxState SLogVisualizer::GetToggleCameraState() const
 {
 	return ALogVisualizerCameraController::IsEnabled(GetWorld())
-		? ESlateCheckBoxState::Checked : ESlateCheckBoxState::Unchecked;
+		? ECheckBoxState::Checked : ECheckBoxState::Unchecked;
 }
 
-void SLogVisualizer::OnOffsetDataSets(ESlateCheckBoxState::Type NewState)
+void SLogVisualizer::OnOffsetDataSets(ECheckBoxState NewState)
 {
-	bOffsetDataSet = (NewState == ESlateCheckBoxState::Checked);
+	bOffsetDataSet = (NewState == ECheckBoxState::Checked);
 	InvalidateCanvas();
 }
 
-ESlateCheckBoxState::Type SLogVisualizer::GetOffsetDataSets() const
+ECheckBoxState SLogVisualizer::GetOffsetDataSets() const
 {
-	return bOffsetDataSet ? ESlateCheckBoxState::Checked : ESlateCheckBoxState::Unchecked;
+	return bOffsetDataSet ? ECheckBoxState::Checked : ECheckBoxState::Unchecked;
 }
 
-void SLogVisualizer::OnHistogramGraphsFilter(ESlateCheckBoxState::Type NewState)
+void SLogVisualizer::OnHistogramGraphsFilter(ECheckBoxState NewState)
 {
-	bHistogramGraphsFilter = (NewState == ESlateCheckBoxState::Checked);
+	bHistogramGraphsFilter = (NewState == ECheckBoxState::Checked);
 	QuickFilterText.Empty();
 	QuickFilterBox->SetText(FText::FromString(QuickFilterText));
 
@@ -2194,9 +2194,9 @@ void SLogVisualizer::OnHistogramGraphsFilter(ESlateCheckBoxState::Type NewState)
 	}
 }
 
-ESlateCheckBoxState::Type SLogVisualizer::GetHistogramGraphsFilter() const
+ECheckBoxState SLogVisualizer::GetHistogramGraphsFilter() const
 {
-	return bHistogramGraphsFilter ? ESlateCheckBoxState::Checked : ESlateCheckBoxState::Unchecked;
+	return bHistogramGraphsFilter ? ECheckBoxState::Checked : ECheckBoxState::Unchecked;
 }
 
 //----------------------------------------------------------------------//
@@ -2609,10 +2609,10 @@ FString SLogVisualizer::GetStatusText() const
 	return TEXT("");
 }
 
-ESlateCheckBoxState::Type SLogVisualizer::GetPauseState() const
+ECheckBoxState SLogVisualizer::GetPauseState() const
 {
 	UWorld* World = GetWorld();
-	return (World != NULL && (World->bPlayersOnly || World->bPlayersOnlyPending)) ? ESlateCheckBoxState::Checked : ESlateCheckBoxState::Unchecked;
+	return (World != NULL && (World->bPlayersOnly || World->bPlayersOnlyPending)) ? ECheckBoxState::Checked : ECheckBoxState::Unchecked;
 }
 
 FReply SLogVisualizer::OnRecordButtonClicked()
@@ -2749,12 +2749,12 @@ FReply SLogVisualizer::OnRemove()
 	return FReply::Handled();
 }
 
-void SLogVisualizer::OnPauseChanged(ESlateCheckBoxState::Type NewState)
+void SLogVisualizer::OnPauseChanged(ECheckBoxState NewState)
 {
 	UWorld* World = GetWorld();
 	if (World != NULL)
 	{
-		if (NewState != ESlateCheckBoxState::Checked)
+		if (NewState != ECheckBoxState::Checked)
 		{
 			World->bPlayersOnly = false;
 			World->bPlayersOnlyPending = false;
diff --git a/Engine/Source/Developer/LogVisualizer/Private/SLogVisualizer.h b/Engine/Source/Developer/LogVisualizer/Private/SLogVisualizer.h
--- a/Engine/Source/Developer/LogVisualizer/Private/SLogVisualizer.h
+++ b/Engine/Source/Developer/LogVisualizer/Private/SLogVisualizer.h
@@ -93,14 +93,14 @@ public:
 	// Get delegates
 	const FSlateBrush* GetRecordButtonBrush() const;
 	FString GetStatusText() const;
-	ESlateCheckBoxState::Type GetPauseState() const;
+	ECheckBoxState GetPauseState() const;
 	EColumnSortMode::Type GetLogsSortMode() const;
 	// Handler delegates
 	FReply OnRecordButtonClicked();
 	FReply OnLoad();
 	FReply OnSave();
 	FReply OnRemove();
-	void OnPauseChanged(ESlateCheckBoxState::Type NewState);
+	void OnPauseChanged(ECheckBoxState NewState);
 	void FilterTextCommitted(const FText& CommentText, ETextCommit::Type CommitInfo);
 	void OnQuickFilterTextChanged(const FText& CommentText, ETextCommit::Type CommitInfo);
 	void OnSortByChanged(const EColumnSortPriority::Type SortPriority, const FName& ColumnName, const EColumnSortMode::Type NewSortMode);
@@ -211,31 +211,31 @@ private:
 	void OnSetHistogramWindowValue(float NewValue);
 
 	/** handler for toggling DrawLogEntriesPath option change */
-	void OnDrawLogEntriesPathChanged(ESlateCheckBoxState::Type NewState);
+	void OnDrawLogEntriesPathChanged(ECheckBoxState NewState);
 	/** retrieves whether DrawLogEntriesPath checkbox should be checked */
-	ESlateCheckBoxState::Type GetDrawLogEntriesPathState() const;
+	ECheckBoxState GetDrawLogEntriesPathState() const;
 
 	/** handler for toggling IgnoreTrivialLogs option change */
-	void OnIgnoreTrivialLogs(ESlateCheckBoxState::Type NewState);
+	void OnIgnoreTrivialLogs(ECheckBoxState NewState);
 	/** retrieves whether IgnoreTrivialLogs checkbox should be checked */
-	ESlateCheckBoxState::Type GetIgnoreTrivialLogs() const;
+	ECheckBoxState GetIgnoreTrivialLogs() const;
 
-	void OnChangeHistogramLabelLocation(ESlateCheckBoxState::Type NewState);
-	ESlateCheckBoxState::Type GetHistogramLabelLocation() const;
+	void OnChangeHistogramLabelLocation(ECheckBoxState NewState);
+	ECheckBoxState GetHistogramLabelLocation() const;
 
-	void OnStickToLastData(ESlateCheckBoxState::Type NewState);
-	ESlateCheckBoxState::Type GetStickToLastData() const;
+	void OnStickToLastData(ECheckBoxState NewState);
+	ECheckBoxState GetStickToLastData() const;
 	
 	/** handler for toggling log visualizer camera in game view */
-	void OnToggleCamera(ESlateCheckBoxState::Type NewState);
+	void OnToggleCamera(ECheckBoxState NewState);
 	/** retrieves whether ToggleCamera toggle button should appear pressed */
-	ESlateCheckBoxState::Type GetToggleCameraState() const;
+	ECheckBoxState GetToggleCameraState() const;
 
-	void OnOffsetDataSets(ESlateCheckBoxState::Type NewState);
-	ESlateCheckBoxState::Type GetOffsetDataSets() const;
+	void OnOffsetDataSets(ECheckBoxState NewState);
+	ECheckBoxState GetOffsetDataSets() const;
 
-	void OnHistogramGraphsFilter(ESlateCheckBoxState::Type NewState);
-	ESlateCheckBoxState::Type GetHistogramGraphsFilter() const;
+	void OnHistogramGraphsFilter(ECheckBoxState NewState);
+	ECheckBoxState GetHistogramGraphsFilter() const;
 	
 	/** See if a particular log passes the current filter and other display flags */
 	bool ShouldListLog(const TSharedPtr<FActorsVisLog>& Log);
diff --git a/Engine/Source/Developer/MessageLog/Private/MessageFilter.cpp b/Engine/Source/Developer/MessageLog/Private/MessageFilter.cpp
--- a/Engine/Source/Developer/MessageLog/Private/MessageFilter.cpp
+++ b/Engine/Source/Developer/MessageLog/Private/MessageFilter.cpp
@@ -9,13 +9,13 @@ FReply FMessageFilter::OnMouseButtonUp(const FGeometry& MyGeometry, const FPoint
 	return FReply::Handled();
 }
 
-ESlateCheckBoxState::Type FMessageFilter::OnGetDisplayCheckState() const
+ECheckBoxState FMessageFilter::OnGetDisplayCheckState() const
 {
-	return Display ? ESlateCheckBoxState::Checked : ESlateCheckBoxState::Unchecked;
+	return Display ? ECheckBoxState::Checked : ECheckBoxState::Unchecked;
 }
 
-void FMessageFilter::OnDisplayCheckStateChanged(ESlateCheckBoxState::Type InNewState)
+void FMessageFilter::OnDisplayCheckStateChanged(ECheckBoxState InNewState)
 {
-	Display = InNewState == ESlateCheckBoxState::Checked;
+	Display = InNewState == ECheckBoxState::Checked;
 	RefreshCallback.Broadcast();
 }
diff --git a/Engine/Source/Developer/MessageLog/Private/MessageFilter.h b/Engine/Source/Developer/MessageLog/Private/MessageFilter.h
--- a/Engine/Source/Developer/MessageLog/Private/MessageFilter.h
+++ b/Engine/Source/Developer/MessageLog/Private/MessageFilter.h
@@ -20,10 +20,10 @@ public:
 	FReply OnMouseButtonUp(const FGeometry& MyGeometry, const FPointerEvent& MouseEvent);
 
 	/** Gets the display state to send to a display filter check box */
-	virtual ESlateCheckBoxState::Type OnGetDisplayCheckState() const;
+	virtual ECheckBoxState OnGetDisplayCheckState() const;
 	
 	/** Sets the display state from a display filter check box */
-	virtual void OnDisplayCheckStateChanged(ESlateCheckBoxState::Type InNewState);
+	virtual void OnDisplayCheckStateChanged(ECheckBoxState InNewState);
 
 	/**
 	 * Returns a delegate to be invoked when the filter state changes.
diff --git a/Engine/Source/Developer/NewLogVisualizer/Private/SFilterWidget.h b/Engine/Source/Developer/NewLogVisualizer/Private/SFilterWidget.h
--- a/Engine/Source/Developer/NewLogVisualizer/Private/SFilterWidget.h
+++ b/Engine/Source/Developer/NewLogVisualizer/Private/SFilterWidget.h
@@ -70,9 +70,9 @@ public:
 
 private:
 	/** Handler for when the filter checkbox is clicked */
-	void FilterToggled(ESlateCheckBoxState::Type NewState)
+	void FilterToggled(ECheckBoxState NewState)
 	{
-		bEnabled = NewState == ESlateCheckBoxState::Checked;
+		bEnabled = NewState == ECheckBoxState::Checked;
 		OnFilterChanged.ExecuteIfBound();
 	}
 
@@ -124,10 +124,10 @@ private:
 	void RemoveAllFilters() { OnRequestRemoveAll.ExecuteIfBound(); }
 
 	/** Handler to determine the "checked" state of the filter checkbox */
-	ESlateCheckBoxState::Type IsChecked() const { return bEnabled ? ESlateCheckBoxState::Checked : ESlateCheckBoxState::Unchecked; }
+	ECheckBoxState IsChecked() const { return bEnabled ? ECheckBoxState::Checked : ECheckBoxState::Unchecked; }
 
 	/** Handler to determine the color of the checkbox when it is checked */
-	FSlateColor GetFilterForegroundColor() const { return IsChecked() ? FilterColor : FLinearColor::White; }
+	FSlateColor GetFilterForegroundColor() const { return IsChecked() == ECheckBoxState::Checked ? FilterColor : FLinearColor::White; }
 
 	/** Handler to determine the padding of the checkbox text when it is pressed */
 	FMargin GetFilterNamePadding() const { return ToggleButtonPtr->IsPressed() ? FMargin(3, 2, 4, 0) : FMargin(3, 1, 4, 1); }
@@ -136,7 +136,7 @@ private:
 	FSlateColor GetFilterNameColorAndOpacity() const
 	{
 		const float DimFactor = 0.75f;
-		return IsChecked() ? (IsHovered() ? ColorCategory * DimFactor : ColorCategory) : (IsHovered() ? FLinearColor::White : FLinearColor::White * DimFactor);
+		return IsChecked() == ECheckBoxState::Checked ? ( IsHovered() ? ColorCategory * DimFactor : ColorCategory ) : ( IsHovered() ? FLinearColor::White : FLinearColor::White * DimFactor );
 	}
 
 private:
diff --git a/Engine/Source/Developer/Profiler/Private/ProfilerManager.h b/Engine/Source/Developer/Profiler/Private/ProfilerManager.h
--- a/Engine/Source/Developer/Profiler/Private/ProfilerManager.h
+++ b/Engine/Source/Developer/Profiler/Private/ProfilerManager.h
@@ -319,7 +319,7 @@ public:
 	}
 
 	// TODO: At this moment SButtonRowBlock::OnIsChecked supports only Checked and Unchecked
-	// ESlateCheckBoxState.Undetermined is not supported
+	// ECheckBoxState.Undetermined is not supported
 	/**
 	 * @return true, if the specified session instance tracks any stats.
 	 */
@@ -638,7 +638,7 @@ protected:
 	/** Profiler view mode. */
 	EProfilerViewMode::Type ViewMode;
 
-	// TODO: Bool should be replaces with type similar to ESlateCheckBoxState {Checked,Unchecked,Undertermined}
+	// TODO: Bool should be replaces with type similar to ECheckBoxState {Checked,Unchecked,Undertermined}
 
 	/** True, if the profiler is currently showing the latest data, only valid if profiler is connected to network based session. */
 	bool bLivePreview;
diff --git a/Engine/Source/Developer/Profiler/Private/Widgets/SEventGraph.cpp b/Engine/Source/Developer/Profiler/Private/Widgets/SEventGraph.cpp
--- a/Engine/Source/Developer/Profiler/Private/Widgets/SEventGraph.cpp
+++ b/Engine/Source/Developer/Profiler/Private/Widgets/SEventGraph.cpp
@@ -1926,25 +1926,25 @@ void SEventGraph::EventGraph_BuildViewColumnMenu( FMenuBuilder& MenuBuilder )
 	MenuBuilder.EndSection();
 }
 
-void SEventGraph::EventGraphViewMode_OnCheckStateChanged( ESlateCheckBoxState::Type NewRadioState, const EEventGraphViewModes::Type InViewMode )
+void SEventGraph::EventGraphViewMode_OnCheckStateChanged( ECheckBoxState NewRadioState, const EEventGraphViewModes::Type InViewMode )
 {
-	if( NewRadioState == ESlateCheckBoxState::Checked && GetCurrentStateViewMode() != InViewMode)
+	if( NewRadioState == ECheckBoxState::Checked && GetCurrentStateViewMode() != InViewMode)
 	{
 		const TArray< FEventGraphSamplePtr > SelectedEvents = TreeView_Base->GetSelectedItems();
 		ShowEventsInViewMode( SelectedEvents, InViewMode );
 	}
 }
 
-ESlateCheckBoxState::Type SEventGraph::EventGraphViewMode_IsChecked( const EEventGraphViewModes::Type InViewMode ) const
+ECheckBoxState SEventGraph::EventGraphViewMode_IsChecked( const EEventGraphViewModes::Type InViewMode ) const
 {
-	return (GetCurrentStateViewMode() == InViewMode) ? ESlateCheckBoxState::Checked : ESlateCheckBoxState::Unchecked;
+	return (GetCurrentStateViewMode() == InViewMode) ? ECheckBoxState::Checked : ECheckBoxState::Unchecked;
 }
 
-void SEventGraph::EventGraphType_OnCheckStateChanged( ESlateCheckBoxState::Type NewRadioState, const EEventGraphTypes::Type NewEventGraphType )
+void SEventGraph::EventGraphType_OnCheckStateChanged( ECheckBoxState NewRadioState, const EEventGraphTypes::Type NewEventGraphType )
 {
 	const uint32 NumFrames = GetCurrentState()->GetNumFrames();
 
-	if( NewRadioState == ESlateCheckBoxState::Checked && GetCurrentStateEventGraphType() != NewEventGraphType )
+	if( NewRadioState == ECheckBoxState::Checked && GetCurrentStateEventGraphType() != NewEventGraphType )
 	{
 		FEventGraphStateRef EventGraphState = GetCurrentState();
 		GetHierarchicalExpandedEvents( EventGraphState->ExpandedEvents );
@@ -1955,7 +1955,7 @@ void SEventGraph::EventGraphType_OnCheckStateChanged( ESlateCheckBoxState::Type
 	}
 }
 
-ESlateCheckBoxState::Type SEventGraph::EventGraphType_IsChecked( const EEventGraphTypes::Type InEventGraphType ) const
+ECheckBoxState SEventGraph::EventGraphType_IsChecked( const EEventGraphTypes::Type InEventGraphType ) const
 {
 	if( IsEventGraphStatesHistoryValid() )
 	{
@@ -1963,15 +1963,15 @@ ESlateCheckBoxState::Type SEventGraph::EventGraphType_IsChecked( const EEventGra
 
 		if( NumFrames == 1 )
 		{
-			return GetCurrentStateEventGraphType() == InEventGraphType ? ESlateCheckBoxState::Checked : ESlateCheckBoxState::Unchecked;
+			return GetCurrentStateEventGraphType() == InEventGraphType ? ECheckBoxState::Checked : ECheckBoxState::Unchecked;
 		}
 		else if( NumFrames > 1 )
 		{
-			return GetCurrentStateEventGraphType() == InEventGraphType ? ESlateCheckBoxState::Checked : ESlateCheckBoxState::Unchecked;
+			return GetCurrentStateEventGraphType() == InEventGraphType ? ECheckBoxState::Checked : ECheckBoxState::Unchecked;
 		}
 	}
 
-	return ESlateCheckBoxState::Unchecked;
+	return ECheckBoxState::Unchecked;
 }
 
 bool SEventGraph::EventGraphType_IsEnabled( const EEventGraphTypes::Type InEventGraphType ) const
@@ -2016,7 +2016,7 @@ FReply SEventGraph::ExpandHotPath_OnClicked()
 	return FReply::Handled();
 }
 
-void SEventGraph::HighlightHotPath_OnCheckStateChanged( ESlateCheckBoxState::Type InState )
+void SEventGraph::HighlightHotPath_OnCheckStateChanged( ECheckBoxState InState )
 {
 
 }
@@ -3323,7 +3323,7 @@ EVisibility SEventGraph::EventGraphViewMode_GetVisibility( const EEventGraphView
 		{
 			// If view mode is not available event graph will switch to the hierarchical view mode.
 			SEventGraph* MutableThis = const_cast< SEventGraph* >( this );
-			MutableThis->EventGraphViewMode_OnCheckStateChanged( ESlateCheckBoxState::Checked, EEventGraphViewModes::Hierarchical );
+			MutableThis->EventGraphViewMode_OnCheckStateChanged( ECheckBoxState::Checked, EEventGraphViewModes::Hierarchical );
 		}
 		return Vis;
 	}
diff --git a/Engine/Source/Developer/Profiler/Private/Widgets/SEventGraph.h b/Engine/Source/Developer/Profiler/Private/Widgets/SEventGraph.h
--- a/Engine/Source/Developer/Profiler/Private/Widgets/SEventGraph.h
+++ b/Engine/Source/Developer/Profiler/Private/Widgets/SEventGraph.h
@@ -247,11 +247,11 @@ protected:
 	void TreeView_Refresh();
 	void SetTreeItemsForViewMode( const EEventGraphViewModes::Type NewViewMode, EEventGraphTypes::Type NewEventGraphType );
 
-	void EventGraphViewMode_OnCheckStateChanged( ESlateCheckBoxState::Type NewRadioState, const EEventGraphViewModes::Type InViewMode );
-	ESlateCheckBoxState::Type EventGraphViewMode_IsChecked( const EEventGraphViewModes::Type InViewMode ) const;
+	void EventGraphViewMode_OnCheckStateChanged( ECheckBoxState NewRadioState, const EEventGraphViewModes::Type InViewMode );
+	ECheckBoxState EventGraphViewMode_IsChecked( const EEventGraphViewModes::Type InViewMode ) const;
 
-	void EventGraphType_OnCheckStateChanged( ESlateCheckBoxState::Type NewRadioState, const EEventGraphTypes::Type NewEventGraphType );
-	ESlateCheckBoxState::Type EventGraphType_IsChecked( const EEventGraphTypes::Type InEventGraphType ) const;
+	void EventGraphType_OnCheckStateChanged( ECheckBoxState NewRadioState, const EEventGraphTypes::Type NewEventGraphType );
+	ECheckBoxState EventGraphType_IsChecked( const EEventGraphTypes::Type InEventGraphType ) const;
 	bool EventGraphType_IsEnabled( const EEventGraphTypes::Type InEventGraphType ) const;
 
 	void FilteringSearchBox_OnTextChanged( const FText& InFilterText );
@@ -262,7 +262,7 @@ protected:
 	void EventGraph_BuildViewColumnMenu( FMenuBuilder& MenuBuilder );
 
 	FReply ExpandHotPath_OnClicked();
-	void HighlightHotPath_OnCheckStateChanged( ESlateCheckBoxState::Type InState );
+	void HighlightHotPath_OnCheckStateChanged( ECheckBoxState InState );
 
 	void InitializeAndShowHeaderColumns();
 	void TreeViewHeaderRow_OnSortModeChanged( const EColumnSortPriority::Type SortPriority, const FName& ColumnID, const EColumnSortMode::Type SortMode );
@@ -483,9 +483,9 @@ public:
 	}
 		
 protected:
-	void HistoryList_GoTo_ExecuteRadioState( ESlateCheckBoxState::Type NewRadioState, int32 StateIndex )
+	void HistoryList_GoTo_ExecuteRadioState( ECheckBoxState NewRadioState, int32 StateIndex )
 	{
-		if( NewRadioState == ESlateCheckBoxState::Checked )
+		if( NewRadioState == ECheckBoxState::Checked )
 		{
 			HistoryList_GoTo_Execute( StateIndex );
 		}
@@ -494,9 +494,9 @@ protected:
 	/** Handles FExecuteAction for HistoryList_GoTo. */
 	void HistoryList_GoTo_Execute( int32 StateIndex );
 
-	ESlateCheckBoxState::Type HistoryList_GoTo_IsCheckedRadioState( int32 StateIndex ) const
+	ECheckBoxState HistoryList_GoTo_IsCheckedRadioState( int32 StateIndex ) const
 	{
-		return HistoryList_GoTo_IsChecked( StateIndex ) ? ESlateCheckBoxState::Checked : ESlateCheckBoxState::Unchecked;
+		return HistoryList_GoTo_IsChecked( StateIndex ) ? ECheckBoxState::Checked : ECheckBoxState::Unchecked;
 	}
 
 	/** Handles FIsActionChecked for HistoryList_GoTo. */
diff --git a/Engine/Source/Developer/Profiler/Private/Widgets/SFiltersAndPresets.cpp b/Engine/Source/Developer/Profiler/Private/Widgets/SFiltersAndPresets.cpp
--- a/Engine/Source/Developer/Profiler/Private/Widgets/SFiltersAndPresets.cpp
+++ b/Engine/Source/Developer/Profiler/Private/Widgets/SFiltersAndPresets.cpp
@@ -1067,15 +1067,15 @@ TSharedRef<SWidget> SFiltersAndPresets::GetToggleButtonForStatType( const EProfi
 		];
 }
 
-void SFiltersAndPresets::FilterByStatType_OnCheckStateChanged( ESlateCheckBoxState::Type NewRadioState, const EProfilerSampleTypes::Type InStatType )
+void SFiltersAndPresets::FilterByStatType_OnCheckStateChanged( ECheckBoxState NewRadioState, const EProfilerSampleTypes::Type InStatType )
 {
-	bStatTypeIsVisible[InStatType] = NewRadioState == ESlateCheckBoxState::Checked;
+	bStatTypeIsVisible[InStatType] = NewRadioState == ECheckBoxState::Checked;
 	ApplyFiltering();
 }
 
-ESlateCheckBoxState::Type SFiltersAndPresets::FilterByStatType_IsChecked( const EProfilerSampleTypes::Type InStatType ) const
+ECheckBoxState SFiltersAndPresets::FilterByStatType_IsChecked( const EProfilerSampleTypes::Type InStatType ) const
 {
-	return bStatTypeIsVisible[InStatType] ? ESlateCheckBoxState::Checked : ESlateCheckBoxState::Unchecked;
+	return bStatTypeIsVisible[InStatType] ? ECheckBoxState::Checked : ECheckBoxState::Unchecked;
 }
 
 /*-----------------------------------------------------------------------------
diff --git a/Engine/Source/Developer/Profiler/Private/Widgets/SFiltersAndPresets.h b/Engine/Source/Developer/Profiler/Private/Widgets/SFiltersAndPresets.h
--- a/Engine/Source/Developer/Profiler/Private/Widgets/SFiltersAndPresets.h
+++ b/Engine/Source/Developer/Profiler/Private/Widgets/SFiltersAndPresets.h
@@ -317,8 +317,8 @@ protected:
 	void RecreateSortByOptionsSources();
 
 	TSharedRef<SWidget> GetToggleButtonForStatType( const EProfilerSampleTypes::Type StatType );
-	void FilterByStatType_OnCheckStateChanged( ESlateCheckBoxState::Type NewRadioState, const EProfilerSampleTypes::Type InStatType );
-	ESlateCheckBoxState::Type FilterByStatType_IsChecked( const EProfilerSampleTypes::Type InStatType ) const;
+	void FilterByStatType_OnCheckStateChanged( ECheckBoxState NewRadioState, const EProfilerSampleTypes::Type InStatType );
+	ECheckBoxState FilterByStatType_IsChecked( const EProfilerSampleTypes::Type InStatType ) const;
 
 	/*-----------------------------------------------------------------------------
 		GroupAndStatTree
diff --git a/Engine/Source/Developer/Profiler/Private/Widgets/SProfilerSettings.cpp b/Engine/Source/Developer/Profiler/Private/Widgets/SProfilerSettings.cpp
--- a/Engine/Source/Developer/Profiler/Private/Widgets/SProfilerSettings.cpp
+++ b/Engine/Source/Developer/Profiler/Private/Widgets/SProfilerSettings.cpp
@@ -222,14 +222,14 @@ FReply SProfilerSettings::ResetToDefaults_OnClicked()
 	return FReply::Handled();
 }
 
-void SProfilerSettings::OptionValue_OnCheckStateChanged( ESlateCheckBoxState::Type CheckBoxState, bool* ValuePtr )
+void SProfilerSettings::OptionValue_OnCheckStateChanged( ECheckBoxState CheckBoxState, bool* ValuePtr )
 {
-	*ValuePtr = CheckBoxState == ESlateCheckBoxState::Checked ? true : false;
+	*ValuePtr = CheckBoxState == ECheckBoxState::Checked ? true : false;
 }
 
-ESlateCheckBoxState::Type SProfilerSettings::OptionValue_IsChecked( const bool* ValuePtr ) const
+ECheckBoxState SProfilerSettings::OptionValue_IsChecked( const bool* ValuePtr ) const
 {
-	return *ValuePtr ? ESlateCheckBoxState::Checked : ESlateCheckBoxState::Unchecked;
+	return *ValuePtr ? ECheckBoxState::Checked : ECheckBoxState::Unchecked;
 }
 
 EVisibility SProfilerSettings::OptionDefault_GetDiffersFromDefaultAsVisibility( const bool* ValuePtr, const bool* DefaultPtr ) const
diff --git a/Engine/Source/Developer/Profiler/Private/Widgets/SProfilerSettings.h b/Engine/Source/Developer/Profiler/Private/Widgets/SProfilerSettings.h
--- a/Engine/Source/Developer/Profiler/Private/Widgets/SProfilerSettings.h
+++ b/Engine/Source/Developer/Profiler/Private/Widgets/SProfilerSettings.h
@@ -36,8 +36,8 @@ private:
 	EVisibility OptionDefault_GetDiffersFromDefaultAsVisibility( const bool* ValuePtr, const bool* DefaultPtr ) const;
 	FReply OptionDefault_OnClicked( bool* ValuePtr, const bool* DefaultPtr );
 
-	void OptionValue_OnCheckStateChanged( ESlateCheckBoxState::Type CheckBoxState, bool* ValuePtr );
-	ESlateCheckBoxState::Type OptionValue_IsChecked( const bool* ValuePtr ) const;
+	void OptionValue_OnCheckStateChanged( ECheckBoxState CheckBoxState, bool* ValuePtr );
+	ECheckBoxState OptionValue_IsChecked( const bool* ValuePtr ) const;
 
 private:
 	/** Delegate to call when this profiler settings is closed. */
diff --git a/Engine/Source/Developer/ProjectLauncher/Private/Widgets/Build/SProjectLauncherBuildPage.cpp b/Engine/Source/Developer/ProjectLauncher/Private/Widgets/Build/SProjectLauncherBuildPage.cpp
--- a/Engine/Source/Developer/ProjectLauncher/Private/Widgets/Build/SProjectLauncherBuildPage.cpp
+++ b/Engine/Source/Developer/ProjectLauncher/Private/Widgets/Build/SProjectLauncherBuildPage.cpp
@@ -166,18 +166,18 @@ bool SProjectLauncherBuildPage::GenerateDSYMForProject( const FString& ProjectNa
 /* SProjectLauncherBuildPage callbacks
  *****************************************************************************/
 
-void SProjectLauncherBuildPage::HandleBuildCheckedStateChanged( ESlateCheckBoxState::Type CheckState )
+void SProjectLauncherBuildPage::HandleBuildCheckedStateChanged( ECheckBoxState CheckState )
 {
 	ILauncherProfilePtr SelectedProfile = Model->GetSelectedProfile();
 
 	if (SelectedProfile.IsValid())
 	{
-		SelectedProfile->SetBuildGame(CheckState == ESlateCheckBoxState::Checked);
+		SelectedProfile->SetBuildGame(CheckState == ECheckBoxState::Checked);
 	}
 }
 
 
-ESlateCheckBoxState::Type SProjectLauncherBuildPage::HandleBuildIsChecked() const
+ECheckBoxState SProjectLauncherBuildPage::HandleBuildIsChecked() const
 {
 	ILauncherProfilePtr SelectedProfile = Model->GetSelectedProfile();
 
@@ -185,11 +185,11 @@ ESlateCheckBoxState::Type SProjectLauncherBuildPage::HandleBuildIsChecked() cons
 	{
 		if (SelectedProfile->IsBuilding())
 		{
-			return ESlateCheckBoxState::Checked;
+			return ECheckBoxState::Checked;
 		}
 	}
 
-	return ESlateCheckBoxState::Unchecked;
+	return ECheckBoxState::Unchecked;
 }
 
 
@@ -304,18 +304,18 @@ EVisibility SProjectLauncherBuildPage::HandleValidationErrorIconVisibility(ELaun
 	return EVisibility::Hidden;
 }
 
-void SProjectLauncherBuildPage::HandleUATCheckedStateChanged( ESlateCheckBoxState::Type CheckState )
+void SProjectLauncherBuildPage::HandleUATCheckedStateChanged( ECheckBoxState CheckState )
 {
 	ILauncherProfilePtr SelectedProfile = Model->GetSelectedProfile();
 
 	if (SelectedProfile.IsValid())
 	{
-		SelectedProfile->SetBuildUAT(CheckState == ESlateCheckBoxState::Checked);
+		SelectedProfile->SetBuildUAT(CheckState == ECheckBoxState::Checked);
 	}
 }
 
 
-ESlateCheckBoxState::Type SProjectLauncherBuildPage::HandleUATIsChecked() const
+ECheckBoxState SProjectLauncherBuildPage::HandleUATIsChecked() const
 {
 	ILauncherProfilePtr SelectedProfile = Model->GetSelectedProfile();
 
@@ -323,11 +323,11 @@ ESlateCheckBoxState::Type SProjectLauncherBuildPage::HandleUATIsChecked() const
 	{
 		if (SelectedProfile->IsBuildingUAT())
 		{
-			return ESlateCheckBoxState::Checked;
+			return ECheckBoxState::Checked;
 		}
 	}
 
-	return ESlateCheckBoxState::Unchecked;
+	return ECheckBoxState::Unchecked;
 }
 
 
diff --git a/Engine/Source/Developer/ProjectLauncher/Private/Widgets/Build/SProjectLauncherBuildPage.h b/Engine/Source/Developer/ProjectLauncher/Private/Widgets/Build/SProjectLauncherBuildPage.h
--- a/Engine/Source/Developer/ProjectLauncher/Private/Widgets/Build/SProjectLauncherBuildPage.h
+++ b/Engine/Source/Developer/ProjectLauncher/Private/Widgets/Build/SProjectLauncherBuildPage.h
@@ -43,10 +43,10 @@ protected:
 private:
 
 	/** Callback for changing the checked state of a platform menu check box. */
-	void HandleBuildCheckedStateChanged( ESlateCheckBoxState::Type CheckState );
+	void HandleBuildCheckedStateChanged( ECheckBoxState CheckState );
 
 	/** Callback for determining whether a platform menu entry is checked. */
-	ESlateCheckBoxState::Type HandleBuildIsChecked() const;
+	ECheckBoxState HandleBuildIsChecked() const;
 
 	/** Callback for changing the selected profile in the profile manager. */
 	void HandleProfileManagerProfileSelected( const ILauncherProfilePtr& SelectedProfile, const ILauncherProfilePtr& PreviousProfile );
@@ -73,10 +73,10 @@ private:
 	EVisibility HandleValidationErrorIconVisibility(ELauncherProfileValidationErrors::Type Error) const;
 
 	/** Callback for changing the checked state of a platform menu check box. */
-	void HandleUATCheckedStateChanged( ESlateCheckBoxState::Type CheckState );
+	void HandleUATCheckedStateChanged( ECheckBoxState CheckState );
 
 	/** Callback for determining whether a platform menu entry is checked. */
-	ESlateCheckBoxState::Type HandleUATIsChecked() const;
+	ECheckBoxState HandleUATIsChecked() const;
 
 private:
 
diff --git a/Engine/Source/Developer/ProjectLauncher/Private/Widgets/Cook/SProjectLauncherCookByTheBookSettings.cpp b/Engine/Source/Developer/ProjectLauncher/Private/Widgets/Cook/SProjectLauncherCookByTheBookSettings.cpp
--- a/Engine/Source/Developer/ProjectLauncher/Private/Widgets/Cook/SProjectLauncherCookByTheBookSettings.cpp
+++ b/Engine/Source/Developer/ProjectLauncher/Private/Widgets/Cook/SProjectLauncherCookByTheBookSettings.cpp
@@ -856,18 +856,18 @@ FString SProjectLauncherCookByTheBookSettings::HandleCookConfigurationSelectorTe
 }
 
 
-void SProjectLauncherCookByTheBookSettings::HandleIncrementalCheckBoxCheckStateChanged( ESlateCheckBoxState::Type NewState )
+void SProjectLauncherCookByTheBookSettings::HandleIncrementalCheckBoxCheckStateChanged( ECheckBoxState NewState )
 {
 	ILauncherProfilePtr SelectedProfile = Model->GetSelectedProfile();
 
 	if (SelectedProfile.IsValid())
 	{
-		SelectedProfile->SetIncrementalCooking(NewState == ESlateCheckBoxState::Checked);
+		SelectedProfile->SetIncrementalCooking(NewState == ECheckBoxState::Checked);
 	}
 }
 
 
-ESlateCheckBoxState::Type SProjectLauncherCookByTheBookSettings::HandleIncrementalCheckBoxIsChecked( ) const
+ECheckBoxState SProjectLauncherCookByTheBookSettings::HandleIncrementalCheckBoxIsChecked( ) const
 {
 	ILauncherProfilePtr SelectedProfile = Model->GetSelectedProfile();
 
@@ -875,11 +875,11 @@ ESlateCheckBoxState::Type SProjectLauncherCookByTheBookSettings::HandleIncrement
 	{
 		if (SelectedProfile->IsCookingIncrementally())
 		{
-			return ESlateCheckBoxState::Checked;
+			return ECheckBoxState::Checked;
 		}
 	}
 
-	return ESlateCheckBoxState::Unchecked;
+	return ECheckBoxState::Unchecked;
 }
 
 
@@ -948,20 +948,20 @@ void SProjectLauncherCookByTheBookSettings::HandleProfileProjectChanged()
 	RefreshCultureList();
 }
 
-ESlateCheckBoxState::Type SProjectLauncherCookByTheBookSettings::HandleShowCheckBoxIsChecked( EShowMapsChoices::Type Choice ) const
+ECheckBoxState SProjectLauncherCookByTheBookSettings::HandleShowCheckBoxIsChecked( EShowMapsChoices::Type Choice ) const
 {
 	if (ShowMapsChoice == Choice)
 	{
-		return ESlateCheckBoxState::Checked;
+		return ECheckBoxState::Checked;
 	}
 
-	return ESlateCheckBoxState::Unchecked;
+	return ECheckBoxState::Unchecked;
 }
 
 
-void SProjectLauncherCookByTheBookSettings::HandleShowCheckBoxCheckStateChanged( ESlateCheckBoxState::Type NewState, EShowMapsChoices::Type Choice )
+void SProjectLauncherCookByTheBookSettings::HandleShowCheckBoxCheckStateChanged( ECheckBoxState NewState, EShowMapsChoices::Type Choice )
 {
-	if (NewState == ESlateCheckBoxState::Checked)
+	if (NewState == ECheckBoxState::Checked)
 	{
 		ShowMapsChoice = Choice;
 		RefreshMapList();
@@ -969,18 +969,18 @@ void SProjectLauncherCookByTheBookSettings::HandleShowCheckBoxCheckStateChanged(
 }
 
 
-void SProjectLauncherCookByTheBookSettings::HandleUnversionedCheckBoxCheckStateChanged( ESlateCheckBoxState::Type NewState )
+void SProjectLauncherCookByTheBookSettings::HandleUnversionedCheckBoxCheckStateChanged( ECheckBoxState NewState )
 {
 	ILauncherProfilePtr SelectedProfile = Model->GetSelectedProfile();
 
 	if (SelectedProfile.IsValid())
 	{
-		SelectedProfile->SetUnversionedCooking((NewState == ESlateCheckBoxState::Checked));
+		SelectedProfile->SetUnversionedCooking((NewState == ECheckBoxState::Checked));
 	}
 }
 
 
-ESlateCheckBoxState::Type SProjectLauncherCookByTheBookSettings::HandleUnversionedCheckBoxIsChecked( ) const
+ECheckBoxState SProjectLauncherCookByTheBookSettings::HandleUnversionedCheckBoxIsChecked( ) const
 {
 	ILauncherProfilePtr SelectedProfile = Model->GetSelectedProfile();
 
@@ -988,11 +988,11 @@ ESlateCheckBoxState::Type SProjectLauncherCookByTheBookSettings::HandleUnversion
 	{
 		if (SelectedProfile->IsCookingUnversioned())
 		{
-			return ESlateCheckBoxState::Checked;
+			return ECheckBoxState::Checked;
 		}
 	}
 
-	return ESlateCheckBoxState::Unchecked;
+	return ECheckBoxState::Unchecked;
 }
 
 
@@ -1045,18 +1045,18 @@ void SProjectLauncherCookByTheBookSettings::HandleCookerOptionsCommitted(const F
 	}
 }
 
-void SProjectLauncherCookByTheBookSettings::HandleUnrealPakCheckBoxCheckStateChanged( ESlateCheckBoxState::Type NewState )
+void SProjectLauncherCookByTheBookSettings::HandleUnrealPakCheckBoxCheckStateChanged( ECheckBoxState NewState )
 {
 	ILauncherProfilePtr SelectedProfile = Model->GetSelectedProfile();
 
 	if (SelectedProfile.IsValid())
 	{
-		SelectedProfile->SetDeployWithUnrealPak(NewState == ESlateCheckBoxState::Checked);
+		SelectedProfile->SetDeployWithUnrealPak(NewState == ECheckBoxState::Checked);
 	}
 }
 
 
-ESlateCheckBoxState::Type SProjectLauncherCookByTheBookSettings::HandleUnrealPakCheckBoxIsChecked( ) const
+ECheckBoxState SProjectLauncherCookByTheBookSettings::HandleUnrealPakCheckBoxIsChecked( ) const
 {
 	ILauncherProfilePtr SelectedProfile = Model->GetSelectedProfile();
 
@@ -1064,11 +1064,11 @@ ESlateCheckBoxState::Type SProjectLauncherCookByTheBookSettings::HandleUnrealPak
 	{
 		if (SelectedProfile->IsPackingWithUnrealPak())
 		{
-			return ESlateCheckBoxState::Checked;
+			return ECheckBoxState::Checked;
 		}
 	}
 
-	return ESlateCheckBoxState::Unchecked;
+	return ECheckBoxState::Unchecked;
 }
 
 #undef LOCTEXT_NAMESPACE
diff --git a/Engine/Source/Developer/ProjectLauncher/Private/Widgets/Cook/SProjectLauncherCookByTheBookSettings.h b/Engine/Source/Developer/ProjectLauncher/Private/Widgets/Cook/SProjectLauncherCookByTheBookSettings.h
--- a/Engine/Source/Developer/ProjectLauncher/Private/Widgets/Cook/SProjectLauncherCookByTheBookSettings.h
+++ b/Engine/Source/Developer/ProjectLauncher/Private/Widgets/Cook/SProjectLauncherCookByTheBookSettings.h
@@ -68,10 +68,10 @@ private:
 	FString HandleCookConfigurationSelectorText( ) const;
 
 	/** Handles check state changes of the 'Incremental' check box. */
-	void HandleIncrementalCheckBoxCheckStateChanged( ESlateCheckBoxState::Type NewState );
+	void HandleIncrementalCheckBoxCheckStateChanged( ECheckBoxState NewState );
 
 	/** Handles determining the checked state of the 'Incremental' check box. */
-	ESlateCheckBoxState::Type HandleIncrementalCheckBoxIsChecked( ) const;
+	ECheckBoxState HandleIncrementalCheckBoxIsChecked( ) const;
 
 	/** Handles generating a row widget in the map list view. */
 	TSharedRef<ITableRow> HandleMapListViewGenerateRow( TSharedPtr<FString> InItem, const TSharedRef<STableViewBase>& OwnerTable );
@@ -92,16 +92,16 @@ private:
 	void HandleProfileManagerProfileSelected( const ILauncherProfilePtr& SelectedProfile, const ILauncherProfilePtr& PreviousProfile );
 
 	/** Handles checking the specified 'Show maps' check box. */
-	void HandleShowCheckBoxCheckStateChanged( ESlateCheckBoxState::Type NewState, EShowMapsChoices::Type Choice );
+	void HandleShowCheckBoxCheckStateChanged( ECheckBoxState NewState, EShowMapsChoices::Type Choice );
 
 	/** Handles determining the checked state of the specified 'Show maps' check box. */
-	ESlateCheckBoxState::Type HandleShowCheckBoxIsChecked( EShowMapsChoices::Type Choice ) const;
+	ECheckBoxState HandleShowCheckBoxIsChecked( EShowMapsChoices::Type Choice ) const;
 
 	/** Handles checking the specified 'Unversioned' check box. */
-	void HandleUnversionedCheckBoxCheckStateChanged( ESlateCheckBoxState::Type NewState );
+	void HandleUnversionedCheckBoxCheckStateChanged( ECheckBoxState NewState );
 
 	/** Handles determining the checked state of the specified 'Unversioned' check box. */
-	ESlateCheckBoxState::Type HandleUnversionedCheckBoxIsChecked( ) const;
+	ECheckBoxState HandleUnversionedCheckBoxIsChecked( ) const;
 
 	/** Handles determining the visibility of a validation error icon. */
 	EVisibility HandleValidationErrorIconVisibility( ELauncherProfileValidationErrors::Type Error ) const;
@@ -116,10 +116,10 @@ private:
 	void HandleProfileProjectChanged();
 
 	// Callback for check state changes of the 'UnrealPak' check box.
-	void HandleUnrealPakCheckBoxCheckStateChanged( ESlateCheckBoxState::Type NewState );
+	void HandleUnrealPakCheckBoxCheckStateChanged( ECheckBoxState NewState );
 
 	// Callback for determining the checked state of the 'UnrealPak' check box.
-	ESlateCheckBoxState::Type HandleUnrealPakCheckBoxIsChecked( ) const;
+	ECheckBoxState HandleUnrealPakCheckBoxIsChecked( ) const;
 
 	/** creates the complex widget. */
 	TSharedRef<SWidget> MakeComplexWidget();
diff --git a/Engine/Source/Developer/ProjectLauncher/Private/Widgets/Cook/SProjectLauncherCookOnTheFlySettings.cpp b/Engine/Source/Developer/ProjectLauncher/Private/Widgets/Cook/SProjectLauncherCookOnTheFlySettings.cpp
--- a/Engine/Source/Developer/ProjectLauncher/Private/Widgets/Cook/SProjectLauncherCookOnTheFlySettings.cpp
+++ b/Engine/Source/Developer/ProjectLauncher/Private/Widgets/Cook/SProjectLauncherCookOnTheFlySettings.cpp
@@ -89,18 +89,18 @@ END_SLATE_FUNCTION_BUILD_OPTIMIZATION
 /* SProjectLauncherCookOnTheFlySettings callbacks
  *****************************************************************************/
 
-void SProjectLauncherCookOnTheFlySettings::HandleIncrementalCheckBoxCheckStateChanged( ESlateCheckBoxState::Type NewState )
+void SProjectLauncherCookOnTheFlySettings::HandleIncrementalCheckBoxCheckStateChanged( ECheckBoxState NewState )
 {
 	ILauncherProfilePtr SelectedProfile = Model->GetSelectedProfile();
 
 	if (SelectedProfile.IsValid())
 	{
-		SelectedProfile->SetIncrementalCooking(NewState == ESlateCheckBoxState::Checked);
+		SelectedProfile->SetIncrementalCooking(NewState == ECheckBoxState::Checked);
 	}
 }
 
 
-ESlateCheckBoxState::Type SProjectLauncherCookOnTheFlySettings::HandleIncrementalCheckBoxIsChecked( ) const
+ECheckBoxState SProjectLauncherCookOnTheFlySettings::HandleIncrementalCheckBoxIsChecked( ) const
 {
 	ILauncherProfilePtr SelectedProfile = Model->GetSelectedProfile();
 
@@ -108,11 +108,11 @@ ESlateCheckBoxState::Type SProjectLauncherCookOnTheFlySettings::HandleIncrementa
 	{
 		if (SelectedProfile->IsCookingIncrementally())
 		{
-			return ESlateCheckBoxState::Checked;
+			return ECheckBoxState::Checked;
 		}
 	}
 
-	return ESlateCheckBoxState::Unchecked;
+	return ECheckBoxState::Unchecked;
 }
 
 
diff --git a/Engine/Source/Developer/ProjectLauncher/Private/Widgets/Cook/SProjectLauncherCookOnTheFlySettings.h b/Engine/Source/Developer/ProjectLauncher/Private/Widgets/Cook/SProjectLauncherCookOnTheFlySettings.h
--- a/Engine/Source/Developer/ProjectLauncher/Private/Widgets/Cook/SProjectLauncherCookOnTheFlySettings.h
+++ b/Engine/Source/Developer/ProjectLauncher/Private/Widgets/Cook/SProjectLauncherCookOnTheFlySettings.h
@@ -34,10 +34,10 @@ public:
 private:
 
 	// Callback for check state changes of the 'Incremental' check box.
-	void HandleIncrementalCheckBoxCheckStateChanged( ESlateCheckBoxState::Type NewState );
+	void HandleIncrementalCheckBoxCheckStateChanged( ECheckBoxState NewState );
 
 	// Callback for determining the checked state of the 'Incremental' check box.
-	ESlateCheckBoxState::Type HandleIncrementalCheckBoxIsChecked( ) const;
+	ECheckBoxState HandleIncrementalCheckBoxIsChecked( ) const;
 
 	// Callback for changing the selected profile in the profile manager.
 	void HandleProfileManagerProfileSelected( const ILauncherProfilePtr& SelectedProfile, const ILauncherProfilePtr& PreviousProfile );
diff --git a/Engine/Source/Developer/ProjectLauncher/Private/Widgets/Cook/SProjectLauncherCultureListRow.h b/Engine/Source/Developer/ProjectLauncher/Private/Widgets/Cook/SProjectLauncherCultureListRow.h
--- a/Engine/Source/Developer/ProjectLauncher/Private/Widgets/Cook/SProjectLauncherCultureListRow.h
+++ b/Engine/Source/Developer/ProjectLauncher/Private/Widgets/Cook/SProjectLauncherCultureListRow.h
@@ -66,13 +66,13 @@ public:
 private:
 
 	// Callback for changing the checked state of the check box.
-	void HandleCheckBoxCheckStateChanged( ESlateCheckBoxState::Type NewState )
+	void HandleCheckBoxCheckStateChanged( ECheckBoxState NewState )
 	{
 		ILauncherProfilePtr SelectedProfile = Model->GetSelectedProfile();
 
 		if (SelectedProfile.IsValid())
 		{
-			if (NewState == ESlateCheckBoxState::Checked)
+			if (NewState == ECheckBoxState::Checked)
 			{
 				SelectedProfile->AddCookedCulture(*CultureName);
 			}
@@ -84,16 +84,16 @@ private:
 	}
 
 	// Callback for determining the checked state of the check box.
-	ESlateCheckBoxState::Type HandleCheckBoxIsChecked( ) const
+	ECheckBoxState HandleCheckBoxIsChecked( ) const
 	{
 		ILauncherProfilePtr SelectedProfile = Model->GetSelectedProfile();
 
 		if (SelectedProfile.IsValid() && SelectedProfile->GetCookedCultures().Contains(*CultureName))
 		{
-			return ESlateCheckBoxState::Checked;
+			return ECheckBoxState::Checked;
 		}
 
-		return ESlateCheckBoxState::Unchecked;
+		return ECheckBoxState::Unchecked;
 	}
 
 private:
diff --git a/Engine/Source/Developer/ProjectLauncher/Private/Widgets/Cook/SProjectLauncherMapListRow.h b/Engine/Source/Developer/ProjectLauncher/Private/Widgets/Cook/SProjectLauncherMapListRow.h
--- a/Engine/Source/Developer/ProjectLauncher/Private/Widgets/Cook/SProjectLauncherMapListRow.h
+++ b/Engine/Source/Developer/ProjectLauncher/Private/Widgets/Cook/SProjectLauncherMapListRow.h
@@ -65,13 +65,13 @@ public:
 private:
 
 	// Callback for changing the checked state of the check box.
-	void HandleCheckBoxCheckStateChanged( ESlateCheckBoxState::Type NewState )
+	void HandleCheckBoxCheckStateChanged( ECheckBoxState NewState )
 	{
 		ILauncherProfilePtr SelectedProfile = Model->GetSelectedProfile();
 
 		if (SelectedProfile.IsValid())
 		{
-			if (NewState == ESlateCheckBoxState::Checked)
+			if (NewState == ECheckBoxState::Checked)
 			{
 				SelectedProfile->AddCookedMap(*MapName);
 			}
@@ -83,16 +83,16 @@ private:
 	}
 
 	// Callback for determining the checked state of the check box.
-	ESlateCheckBoxState::Type HandleCheckBoxIsChecked( ) const
+	ECheckBoxState HandleCheckBoxIsChecked( ) const
 	{
 		ILauncherProfilePtr SelectedProfile = Model->GetSelectedProfile();
 
 		if (SelectedProfile.IsValid() && SelectedProfile->GetCookedMaps().Contains(*MapName))
 		{
-			return ESlateCheckBoxState::Checked;
+			return ECheckBoxState::Checked;
 		}
 
-		return ESlateCheckBoxState::Unchecked;
+		return ECheckBoxState::Unchecked;
 	}
 
 private:
