Moved left scrollbar outside timelines in LogVisualizer. Sometimes it was hard to select correct item on timeline because of this scroll bar.

diff --git a/Engine/Source/Developer/LogVisualizer/Private/SVisualLoggerView.cpp b/Engine/Source/Developer/LogVisualizer/Private/SVisualLoggerView.cpp
--- a/Engine/Source/Developer/LogVisualizer/Private/SVisualLoggerView.cpp
+++ b/Engine/Source/Developer/LogVisualizer/Private/SVisualLoggerView.cpp
@@ -109,10 +109,9 @@ void SVisualLoggerView::Construct(const FArguments& InArgs, const TSharedRef<FUI
 			.Padding(2)
 			.BorderImage(FLogVisualizerStyle::Get().GetBrush("ToolPanel.GroupBorder"))
 			[
-				SNew(SSplitter)
-				.Orientation(Orient_Vertical)
-				.PhysicalSplitterHandleSize(2)
-				+ SSplitter::Slot()
+				SNew(SHorizontalBox)
+				+ SHorizontalBox::Slot()
+				.FillWidth(1.0f)
 				[
 					SNew(SVerticalBox)
 					+ SVerticalBox::Slot()
@@ -188,11 +187,6 @@ void SVisualLoggerView::Construct(const FArguments& InArgs, const TSharedRef<FUI
 							MakeSectionOverlay(TimeSliderController.ToSharedRef(), InArgs._ViewRange, InArgs._ScrubPosition, true)
 						]
 						+ SOverlay::Slot()
-						.HAlign(HAlign_Right)
-						[
-							ScrollBar
-						]
-						+ SOverlay::Slot()
 						.VAlign(VAlign_Bottom)
 						[
 							SNew(SHorizontalBox)
@@ -234,6 +228,12 @@ void SVisualLoggerView::Construct(const FArguments& InArgs, const TSharedRef<FUI
 						]
 					]
 				]
+				+ SHorizontalBox::Slot()
+				.HAlign(HAlign_Right)
+				.AutoWidth()
+				[
+					ScrollBar
+				]
 			]
 		];
 
