Bunch of fixes for LogVisualizer and added two new features to this tool:
- Filters can be presistent even between editor sessions (option for this in settings)
- Option to filter logs by owner class

diff --git a/Engine/Source/Developer/LogVisualizer/Private/LogVisualizer.h b/Engine/Source/Developer/LogVisualizer/Private/LogVisualizer.h
--- a/Engine/Source/Developer/LogVisualizer/Private/LogVisualizer.h
+++ b/Engine/Source/Developer/LogVisualizer/Private/LogVisualizer.h
@@ -13,6 +13,8 @@
 /* Private includes
 *****************************************************************************/
 #include "VisualLogger/VisualLogger.h"
+#include "LogVisualizerSettings.h"
+#include "LogVisualizerSessionSettings.h"
 
 DECLARE_DELEGATE_OneParam(FOnItemSelectionChanged, const FVisualLogDevice::FVisualLogEntryItem&);
 DECLARE_DELEGATE_OneParam(FOnObjectSelectionChanged, TSharedPtr<class STimeline>);
@@ -44,8 +46,10 @@ struct IVisualLoggerInterface
 	const FVisualLoggerEvents& GetVisualLoggerEvents() { return VisualLoggerEvents; }
 	class FSequencerTimeSliderController* GetTimeSliderController() { return TimeSliderController.Get(); }
 	void SetTimeSliderController(TSharedPtr<class FSequencerTimeSliderController> InTimeSliderController) { TimeSliderController = InTimeSliderController; }
+	FFiltersPreset& GetFiltersPreset() { return FiltersPreset; }
 
 protected:
+	FFiltersPreset	FiltersPreset;
 	FVisualLoggerEvents VisualLoggerEvents;
 	TSharedPtr<class FSequencerTimeSliderController> TimeSliderController;
 };
@@ -60,5 +64,3 @@ protected:
 #include "SVisualLoggerLogsList.h"
 #include "SVisualLoggerStatusView.h"
 #include "STimeline.h"
-#include "LogVisualizerSettings.h"
-#include "LogVisualizerSessionSettings.h"
\ No newline at end of file
diff --git a/Engine/Source/Developer/LogVisualizer/Private/LogVisualizerModule.cpp b/Engine/Source/Developer/LogVisualizer/Private/LogVisualizerModule.cpp
--- a/Engine/Source/Developer/LogVisualizer/Private/LogVisualizerModule.cpp
+++ b/Engine/Source/Developer/LogVisualizer/Private/LogVisualizerModule.cpp
@@ -19,7 +19,7 @@ static const FName VisualLoggerTabName("VisualLogger");
 
 //DEFINE_LOG_CATEGORY(LogLogVisualizer);
 
-void FNewLogVisualizerModule::StartupModule()
+void FLogVisualizerModule::StartupModule()
 {
 	FLogVisualizerStyle::Initialize();
 
@@ -28,7 +28,7 @@ void FNewLogVisualizerModule::StartupModule()
 
 	FGlobalTabmanager::Get()->RegisterNomadTabSpawner(
 		VisualLoggerTabName, 
-		FOnSpawnTab::CreateRaw(this, &FNewLogVisualizerModule::SpawnLogVisualizerTab))
+		FOnSpawnTab::CreateRaw(this, &FLogVisualizerModule::SpawnLogVisualizerTab))
 		.SetGroup(WorkspaceMenu::GetMenuStructure().GetDeveloperToolsMiscCategory())
 		.SetDisplayName(NSLOCTEXT("LogVisualizerApp", "TabTitle", "Visual Logger"))
 		.SetTooltipText(NSLOCTEXT("LogVisualizerApp", "TooltipText", "Opens Visual Logger tool."))
@@ -45,7 +45,7 @@ void FNewLogVisualizerModule::StartupModule()
 	}
 }
 
-void FNewLogVisualizerModule::ShutdownModule()
+void FLogVisualizerModule::ShutdownModule()
 {
 	FGlobalTabmanager::Get()->UnregisterTabSpawner(VisualLoggerTabName);
 	FVisualLoggerCommands::Unregister();
@@ -60,10 +60,10 @@ void FNewLogVisualizerModule::ShutdownModule()
 	FLogVisualizerStyle::Shutdown();
 }
 
-TSharedRef<SDockTab> FNewLogVisualizerModule::SpawnLogVisualizerTab(const FSpawnTabArgs& SpawnTabArgs)
+TSharedRef<SDockTab> FLogVisualizerModule::SpawnLogVisualizerTab(const FSpawnTabArgs& SpawnTabArgs)
 {
 	const TSharedRef<SDockTab> MajorTab = SNew(SDockTab)
-		.TabRole(ETabRole::MajorTab).OnTabClosed(SDockTab::FOnTabClosedCallback::CreateRaw(this, &FNewLogVisualizerModule::OnTabClosed));
+		.TabRole(ETabRole::MajorTab);
 
 	TSharedPtr<SWidget> TabContent;
 
@@ -74,41 +74,5 @@ TSharedRef<SDockTab> FNewLogVisualizerModule::SpawnLogVisualizerTab(const FSpawn
 	return MajorTab;
 }
 
-void FNewLogVisualizerModule::OnTabClosed(TSharedRef<SDockTab> DockTab)
-{
-	UWorld* World = NULL;
-#if WITH_EDITOR
-	UEditorEngine *EEngine = Cast<UEditorEngine>(GEngine);
-	if (GIsEditor && EEngine != NULL)
-	{
-		// lets use PlayWorld during PIE/Simulate and regular world from editor otherwise, to draw debug information
-		World = EEngine->PlayWorld != NULL ? EEngine->PlayWorld : EEngine->GetEditorWorldContext().World();
-
-	}
-	else 
-#endif
-	if (!GIsEditor)
-	{
-
-		World = GEngine->GetWorld();
-	}
-
-	TSharedRef<SVisualLogger> VisualLoggerTab = StaticCastSharedRef<SVisualLogger>(DockTab->GetContent());
-	VisualLoggerTab->OnTabLosed();
-
-	if (World == NULL)
-	{
-		World = GWorld;
-	}
-
-	if (World)
-	{
-		for (TActorIterator<AVisualLoggerRenderingActor> It(World); It; ++It)
-		{
-			World->DestroyActor(*It);
-		}
-	}
-}
-
-IMPLEMENT_MODULE(FNewLogVisualizerModule, LogVisualizer);
+IMPLEMENT_MODULE(FLogVisualizerModule, LogVisualizer);
 #undef LOCTEXT_NAMESPACE
diff --git a/Engine/Source/Developer/LogVisualizer/Private/LogVisualizerStyle.cpp b/Engine/Source/Developer/LogVisualizer/Private/LogVisualizerStyle.cpp
--- a/Engine/Source/Developer/LogVisualizer/Private/LogVisualizerStyle.cpp
+++ b/Engine/Source/Developer/LogVisualizer/Private/LogVisualizerStyle.cpp
@@ -47,6 +47,14 @@ TSharedRef< FSlateStyleSet > FLogVisualizerStyle::Create()
 
 	FSlateStyleSet& Style = StyleRef.Get();
 	// Generic styles
+	const FTextBlockStyle NormalText = FTextBlockStyle()
+		.SetFont(TTF_FONT("Fonts/Roboto-Regular", 9))
+		.SetColorAndOpacity(FSlateColor::UseForeground())
+		.SetShadowOffset(FVector2D::ZeroVector)
+		.SetShadowColorAndOpacity(FLinearColor::Black)
+		.SetHighlightColor(FLinearColor(0.02f, 0.3f, 0.0f))
+		.SetHighlightShape(BOX_BRUSH("Common/TextBlockHighlightShape", FMargin(3.f / 8.f)));
+
 	{
 		Style.Set("LogVisualizerApp.TabIcon", new IMAGE_BRUSH("Icons/icon_tab_DebugTools_40x", Icon16x16));
 		
@@ -108,18 +116,13 @@ TSharedRef< FSlateStyleSet > FLogVisualizerStyle::Create()
 		Style.Set("Sequencer.ItemTitle.Normal", new BOX_BRUSH("Common/Button/simple_round_normal", FMargin(4 / 16.0f), FLinearColor(1, 1, 1, 1)));
 		Style.Set("Sequencer.ItemTitle.Hover", new BOX_BRUSH("Common/Button/simple_round_hovered", FMargin(4 / 16.0f), FLinearColor(1, 1, 1, 1)));
 		Style.Set("Sequencer.SectionArea.Background", new FSlateColorBrush(FColor::White));
+		const FTextBlockStyle CLassNameTextStyle = FTextBlockStyle(NormalText)
+			.SetFont(TTF_FONT("Fonts/Roboto-Regular", 7));
+		Style.Set("Sequencer.ClassNAme", CLassNameTextStyle);
 	}
 
 	// Default text styles
 	{
-		const FTextBlockStyle NormalText = FTextBlockStyle()
-			.SetFont(TTF_FONT("Fonts/Roboto-Regular", 9))
-			.SetColorAndOpacity(FSlateColor::UseForeground())
-			.SetShadowOffset(FVector2D::ZeroVector)
-			.SetShadowColorAndOpacity(FLinearColor::Black)
-			.SetHighlightColor(FLinearColor(0.02f, 0.3f, 0.0f))
-			.SetHighlightShape(BOX_BRUSH("Common/TextBlockHighlightShape", FMargin(3.f / 8.f)));
-
 		Style.Set("RichText.Background", new BOX_BRUSH("Common/FlatColorSquare", FVector2D(1.0f, 1.0f), FMargin(0), FLinearColor(FColor(0xffeff3f3))));
 
 		Style.Set("RichText.RoundedBackground", new BOX_BRUSH("Common/RoundedSelection_16x", 4.0f / 16.0f, FLinearColor(FColor(0xffeff3f3))));
diff --git a/Engine/Source/Developer/LogVisualizer/Private/NewLogVisualizer.cpp b/Engine/Source/Developer/LogVisualizer/Private/NewLogVisualizer.cpp
deleted file mode 100644
diff --git a/Engine/Source/Developer/LogVisualizer/Private/NewLogVisualizer.h b/Engine/Source/Developer/LogVisualizer/Private/NewLogVisualizer.h
deleted file mode 100644
diff --git a/Engine/Source/Developer/LogVisualizer/Private/STimeline.cpp b/Engine/Source/Developer/LogVisualizer/Private/STimeline.cpp
--- a/Engine/Source/Developer/LogVisualizer/Private/STimeline.cpp
+++ b/Engine/Source/Developer/LogVisualizer/Private/STimeline.cpp
@@ -76,8 +76,10 @@ void STimeline::UpdateVisibility()
 {
 	ULogVisualizerSettings* Settings = ULogVisualizerSettings::StaticClass()->GetDefaultObject<ULogVisualizerSettings>();
 	const bool bJustIgnore = Settings->bIgnoreTrivialLogs && Entries.Num() <= Settings->TrivialLogsThreshold;
+	const bool bVisibleByOwnerClasses = Settings->CurrentPresets.SelectedClasses.Num() == 0 || Settings->CurrentPresets.SelectedClasses.Find(OwnerClassName.ToString()) != INDEX_NONE;
 	const bool bIsCollapsed = bJustIgnore || HiddenEntries.Num() == Entries.Num() || (SearchFilter.Len() > 0 && Name.ToString().Find(SearchFilter) == INDEX_NONE);
-	SetVisibility(bIsCollapsed ? EVisibility::Collapsed : EVisibility::Visible);
+
+	SetVisibility(bIsCollapsed || !bVisibleByOwnerClasses ? EVisibility::Collapsed : EVisibility::Visible);
 	if (bIsCollapsed)
 	{
 		Owner->SetSelectionState(SharedThis(this), false, false);
@@ -87,54 +89,18 @@ void STimeline::UpdateVisibility()
 void STimeline::UpdateVisibilityForItems()
 {
 	HiddenEntries.Reset();
+	ULogVisualizerSettings* Settings = ULogVisualizerSettings::StaticClass()->GetDefaultObject<ULogVisualizerSettings>();
+	const FString QuickSearchStrng = Settings->CurrentPresets.DataFilter;
+
 	for (auto& CurrentEntry : Entries)
 	{
-		TArray<FVisualLoggerCategoryVerbosityPair> OutCategories;
-		FVisualLoggerHelpers::GetCategories(CurrentEntry.Entry, OutCategories);
-		bool bHasValidCategories = VisualLoggerInterface->HasValidCategories(OutCategories);
-
-		const ULogVisualizerSettings* Settings = ULogVisualizerSettings::StaticClass()->GetDefaultObject<ULogVisualizerSettings>();
-		if (Settings->bSearchInsideLogs && bHasValidCategories && QuickSearchStrng.Len() > 0)
-		{
-			bool bMatchSearchString = false;
-			for (const FVisualLogLine& CurrentLine : CurrentEntry.Entry.LogLines)
-			{
-				if (CurrentLine.Line.Find(QuickSearchStrng) != INDEX_NONE || CurrentLine.Category.ToString().Find(QuickSearchStrng) != INDEX_NONE)
-				{
-					bMatchSearchString = true;
-					break;
-				}
-			}
-			if (!bMatchSearchString)
-			{
-				for (const FVisualLogEvent& CurrentEvent : CurrentEntry.Entry.Events)
-				{
-					if (CurrentEvent.Name.Find(QuickSearchStrng) != INDEX_NONE)
-					{
-						bMatchSearchString = true;
-						break;
-					}
-				}
-			}
-
-			if (!bMatchSearchString)
-			{
-				HiddenEntries.AddUnique(&CurrentEntry);
-			}
-		}
-		else
-		{
-			if (bHasValidCategories == false)
-			{
-				HiddenEntries.AddUnique(&CurrentEntry);
-			}
-		}
+		UpdateVisibilityForEntry(CurrentEntry, QuickSearchStrng, Settings->bSearchInsideLogs);
 	}
 }
 
 void STimeline::OnFiltersSearchChanged(const FText& Filter)
 {
-	QuickSearchStrng = Filter.ToString();
+	//QuickSearchStrng = Filter.ToString();
 	OnFiltersChanged();
 }
 
@@ -160,6 +126,58 @@ void STimeline::HandleLogVisualizerSettingChanged(FName Name)
 	UpdateVisibility();
 }
 
+void STimeline::UpdateVisibilityForEntry(FVisualLogDevice::FVisualLogEntryItem& CurrentEntry, const FString& SearchString, bool bSearchInsideLogs)
+{
+	TArray<FVisualLoggerCategoryVerbosityPair> OutCategories;
+	FVisualLoggerHelpers::GetCategories(CurrentEntry.Entry, OutCategories);
+	bool bHasValidCategories = VisualLoggerInterface.Pin()->HasValidCategories(OutCategories);
+
+	if (bSearchInsideLogs && bHasValidCategories && SearchString.Len() > 0)
+	{
+		bool bMatchSearchString = false;
+		for (const FVisualLogLine& CurrentLine : CurrentEntry.Entry.LogLines)
+		{
+			if (CurrentLine.Line.Find(SearchString) != INDEX_NONE || CurrentLine.Category.ToString().Find(SearchString) != INDEX_NONE)
+			{
+				bMatchSearchString = true;
+				break;
+			}
+		}
+		if (!bMatchSearchString)
+		{
+			for (const FVisualLogEvent& CurrentEvent : CurrentEntry.Entry.Events)
+			{
+				if (CurrentEvent.Name.Find(SearchString) != INDEX_NONE)
+				{
+					bMatchSearchString = true;
+					break;
+				}
+			}
+		}
+
+		if (!bMatchSearchString)
+		{
+			HiddenEntries.AddUnique(&CurrentEntry);
+		}
+	}
+	else
+	{
+		if (bHasValidCategories == false)
+		{
+			HiddenEntries.AddUnique(&CurrentEntry);
+		}
+	}
+}
+
+void STimeline::AddEntry(const FVisualLogDevice::FVisualLogEntryItem& Entry) 
+{ 
+	Entries.Add(Entry); 
+
+	ULogVisualizerSettings* Settings = ULogVisualizerSettings::StaticClass()->GetDefaultObject<ULogVisualizerSettings>();
+	UpdateVisibilityForEntry(Entries[Entries.Num() - 1], Settings->CurrentPresets.DataFilter, Settings->bSearchInsideLogs);
+	UpdateVisibility();
+}
+
 void STimeline::Construct(const FArguments& InArgs, TSharedPtr<SVisualLoggerView> VisualLoggerView, TSharedPtr<FSequencerTimeSliderController> TimeSliderController, TSharedPtr<STimelinesContainer> InContainer, const FVisualLogDevice::FVisualLogEntryItem& Entry)
 {
 	VisualLoggerInterface = InArgs._VisualLoggerInterface.Get();
@@ -167,6 +185,7 @@ void STimeline::Construct(const FArguments& InArgs, TSharedPtr<SVisualLoggerView
 
 	Owner = InContainer;
 	Name = Entry.OwnerName;
+	OwnerClassName = Entry.OwnerClassName;
 
 	Entries.Add(Entry);
 	OnFiltersChanged();
@@ -197,9 +216,25 @@ void STimeline::Construct(const FArguments& InArgs, TSharedPtr<SVisualLoggerView
 						.Padding(FMargin(4, 0, 2, 0))
 						[
 							// Search box for searching through the outliner
-							SNew(STextBlock)
-							.Text(FString(Name.ToString()))
-							.ShadowOffset(FVector2D(1.f, 1.f))
+							SNew(SVerticalBox)
+							+ SVerticalBox::Slot()
+							.Padding(FMargin(0, 0, 0, 0))
+							.HAlign(HAlign_Fill)
+							.VAlign(VAlign_Fill)
+							[
+								SNew(STextBlock)
+								.Text( FText::FromName(Name) )
+								.ShadowOffset(FVector2D(1.f, 1.f))
+							]
+							+ SVerticalBox::Slot()
+							.Padding(FMargin(0, 0, 0, 0))
+							.HAlign(HAlign_Fill)
+							.VAlign(VAlign_Fill)
+							[
+								SNew(STextBlock)
+								.Text(FText::FromName(OwnerClassName))
+								.TextStyle(FLogVisualizerStyle::Get(), TEXT("Sequencer.ClassNAme"))
+							]
 						]
 					]
 				]
diff --git a/Engine/Source/Developer/LogVisualizer/Private/STimeline.h b/Engine/Source/Developer/LogVisualizer/Private/STimeline.h
--- a/Engine/Source/Developer/LogVisualizer/Private/STimeline.h
+++ b/Engine/Source/Developer/LogVisualizer/Private/STimeline.h
@@ -9,7 +9,7 @@ class STimeline : public SCompoundWidget
 {
 public:
 	SLATE_BEGIN_ARGS(STimeline){}
-		SLATE_ATTRIBUTE(TSharedPtr<IVisualLoggerInterface>, VisualLoggerInterface)
+		SLATE_ATTRIBUTE(TWeakPtr<IVisualLoggerInterface>, VisualLoggerInterface)
 		SLATE_EVENT(FOnItemSelectionChanged, OnItemSelectionChanged)
 		SLATE_EVENT(FOnGetContent, OnGetMenuContent)
 	SLATE_END_ARGS();
@@ -33,22 +33,24 @@ public:
 	void UpdateVisibility();
 	void UpdateVisibilityForItems();
 
-	void AddEntry(const FVisualLogDevice::FVisualLogEntryItem& Entry) { Entries.Add(Entry); }
+	void AddEntry(const FVisualLogDevice::FVisualLogEntryItem& Entry);
 	const TArray<FVisualLogDevice::FVisualLogEntryItem>& GetEntries() { return Entries; }
 	FName GetName() { return Name; }
+	FName GetOwnerClassName() { return OwnerClassName; }
 	void HandleLogVisualizerSettingChanged(FName Name);
+	void UpdateVisibilityForEntry(FVisualLogDevice::FVisualLogEntryItem& Entry, const FString& SearchString, bool SearchInsideLogs);
 
 protected:
 	TArray<FVisualLogDevice::FVisualLogEntryItem> Entries;
 	TArray<const FVisualLogDevice::FVisualLogEntryItem*> HiddenEntries;
 	TSharedPtr<class STimelinesContainer> Owner;
 	TSharedPtr<class STimelineBar> TimelineBar;
-	TSharedPtr<IVisualLoggerInterface> VisualLoggerInterface;
+	TWeakPtr<IVisualLoggerInterface> VisualLoggerInterface;
 	TSharedPtr<SMenuAnchor> PopupAnchor;
 
 	FName Name;
+	FName OwnerClassName;
 	FString SearchFilter;
-	FString QuickSearchStrng;
 
 	/** Delegate to execute to get the menu content of this timeline */
 	FOnGetContent OnGetMenuContent;
diff --git a/Engine/Source/Developer/LogVisualizer/Private/STimelineBar.cpp b/Engine/Source/Developer/LogVisualizer/Private/STimelineBar.cpp
--- a/Engine/Source/Developer/LogVisualizer/Private/STimelineBar.cpp
+++ b/Engine/Source/Developer/LogVisualizer/Private/STimelineBar.cpp
@@ -43,7 +43,7 @@ FReply STimelineBar::OnMouseButtonDoubleClick(const FGeometry& MyGeometry, const
 {
 	auto &Entries = TimelineOwner.Pin()->GetEntries();
 	SnapScrubPosition(TimeSliderController->GetTimeSliderArgs().ScrubPosition.Get());
-	UWorld* World = VisualLoggerInterface->GetWorld();
+	UWorld* World = VisualLoggerInterface.Pin()->GetWorld();
 	if (World && MouseEvent.GetEffectingButton() == EKeys::LeftMouseButton && Entries.IsValidIndex(CurrentItemIndex))
 	{
 		FVector CurrentLocation = Entries[CurrentItemIndex].Entry.Location;
@@ -233,7 +233,7 @@ void STimelineBar::SnapScrubPosition(float ScrubPosition)
 		if (CurrentItemIndex != BestItemIndex)
 		{
 			CurrentItemIndex = BestItemIndex;
-			VisualLoggerInterface->GetVisualLoggerEvents().OnItemSelectionChanged.ExecuteIfBound(Entries[BestItemIndex]);
+			VisualLoggerInterface.Pin()->GetVisualLoggerEvents().OnItemSelectionChanged.ExecuteIfBound(Entries[BestItemIndex]);
 		}
 		TimeSliderController->CommitScrubPosition(CurrentTime, false);
 	}
@@ -248,7 +248,7 @@ void STimelineBar::SnapScrubPosition(int32 NewItemIndex)
 		if (CurrentItemIndex != NewItemIndex)
 		{
 			CurrentItemIndex = NewItemIndex;
-			VisualLoggerInterface->GetVisualLoggerEvents().OnItemSelectionChanged.ExecuteIfBound(Entries[NewItemIndex]);
+			VisualLoggerInterface.Pin()->GetVisualLoggerEvents().OnItemSelectionChanged.ExecuteIfBound(Entries[NewItemIndex]);
 		}
 		TimeSliderController->CommitScrubPosition(CurrentTime, false);
 	}
@@ -440,7 +440,7 @@ int32 STimelineBar::OnPaint(const FPaintArgs& Args, const FGeometry& AllottedGeo
 		auto &Entries = TimelineOwner.Pin()->GetEntries();
 		if (BestItemIndex != CurrentItemIndex)
 		{
-			VisualLoggerInterface->GetVisualLoggerEvents().OnItemSelectionChanged.ExecuteIfBound(Entries[BestItemIndex]);
+			VisualLoggerInterface.Pin()->GetVisualLoggerEvents().OnItemSelectionChanged.ExecuteIfBound(Entries[BestItemIndex]);
 		}
 
 		float CurrentTime = Entries[BestItemIndex].Entry.TimeStamp;
diff --git a/Engine/Source/Developer/LogVisualizer/Private/STimelineBar.h b/Engine/Source/Developer/LogVisualizer/Private/STimelineBar.h
--- a/Engine/Source/Developer/LogVisualizer/Private/STimelineBar.h
+++ b/Engine/Source/Developer/LogVisualizer/Private/STimelineBar.h
@@ -11,7 +11,7 @@ class STimelineBar : public SLeafWidget
 public:
 	SLATE_BEGIN_ARGS(STimelineBar){}
 		SLATE_EVENT(FOnItemSelectionChanged, OnItemSelectionChanged)
-		SLATE_ATTRIBUTE(TSharedPtr<IVisualLoggerInterface>, VisualLoggerInterface)
+		SLATE_ATTRIBUTE(TWeakPtr<IVisualLoggerInterface>, VisualLoggerInterface)
 	SLATE_END_ARGS()
 
 	/** Constructs this widget with InArgs */
@@ -34,7 +34,7 @@ public:
 	void OnDeselect();
 
 protected:
-	TSharedPtr<IVisualLoggerInterface> VisualLoggerInterface;
+	TWeakPtr<IVisualLoggerInterface> VisualLoggerInterface;
 	TSharedPtr<class FSequencerTimeSliderController> TimeSliderController;
 	TWeakPtr<class STimeline> TimelineOwner;
 	mutable int32 CurrentItemIndex;
diff --git a/Engine/Source/Developer/LogVisualizer/Private/STimelinesContainer.cpp b/Engine/Source/Developer/LogVisualizer/Private/STimelinesContainer.cpp
--- a/Engine/Source/Developer/LogVisualizer/Private/STimelinesContainer.cpp
+++ b/Engine/Source/Developer/LogVisualizer/Private/STimelinesContainer.cpp
@@ -16,7 +16,7 @@ TSharedRef<SWidget> STimelinesContainer::MakeTimeline(TSharedPtr<class SVisualLo
 	ContainingBorder->AddSlot()
 		[
 			SAssignNew(NewTimeline, STimeline, VisualLoggerView, TimeSliderController, SharedThis(this), Entry)
-			.OnItemSelectionChanged(this->VisualLoggerInterface->GetVisualLoggerEvents().OnItemSelectionChanged)
+			.OnItemSelectionChanged(this->VisualLoggerInterface.Pin()->GetVisualLoggerEvents().OnItemSelectionChanged)
 			.VisualLoggerInterface(this->VisualLoggerInterface)
 			.OnGetMenuContent(this, &STimelinesContainer::GetRightClickMenuContent)
 		];
@@ -72,7 +72,7 @@ void STimelinesContainer::SetSelectionState(TSharedPtr<class STimeline> Affected
 			SelectedNodes.Add(AffectedNode);
 			AffectedNode->OnSelect();
 		}
-		VisualLoggerInterface->GetVisualLoggerEvents().OnObjectSelectionChanged.ExecuteIfBound(AffectedNode);
+		VisualLoggerInterface.Pin()->GetVisualLoggerEvents().OnObjectSelectionChanged.ExecuteIfBound(AffectedNode);
 	}
 	else if (AffectedNode.IsValid())
 	{
diff --git a/Engine/Source/Developer/LogVisualizer/Private/STimelinesContainer.h b/Engine/Source/Developer/LogVisualizer/Private/STimelinesContainer.h
--- a/Engine/Source/Developer/LogVisualizer/Private/STimelinesContainer.h
+++ b/Engine/Source/Developer/LogVisualizer/Private/STimelinesContainer.h
@@ -9,7 +9,7 @@ class STimelinesContainer : public SCompoundWidget
 {
 public:
 	SLATE_BEGIN_ARGS(STimelinesContainer){}
-		SLATE_ATTRIBUTE(TSharedPtr<IVisualLoggerInterface>, VisualLoggerInterface)
+		SLATE_ATTRIBUTE(TWeakPtr<IVisualLoggerInterface>, VisualLoggerInterface)
 	SLATE_END_ARGS()
 
 	virtual FReply OnMouseButtonDown(const FGeometry& MyGeometry, const FPointerEvent& MouseEvent) override;
@@ -64,5 +64,5 @@ protected:
 	TArray<TSharedPtr<class STimeline> > TimelineItems;
 	TArray< TSharedPtr<class STimeline> > SelectedNodes;
 	TSharedPtr<SVerticalBox> ContainingBorder;
-	TSharedPtr<IVisualLoggerInterface> VisualLoggerInterface;
+	TWeakPtr<IVisualLoggerInterface> VisualLoggerInterface;
 };
diff --git a/Engine/Source/Developer/LogVisualizer/Private/SVisualLogger.cpp b/Engine/Source/Developer/LogVisualizer/Private/SVisualLogger.cpp
--- a/Engine/Source/Developer/LogVisualizer/Private/SVisualLogger.cpp
+++ b/Engine/Source/Developer/LogVisualizer/Private/SVisualLogger.cpp
@@ -46,13 +46,16 @@ struct FVisualLoggerInterface : public IVisualLoggerInterface
 		VisualLoggerEvents = InVisualLoggerEvents;
 	}
 
-	virtual ~FVisualLoggerInterface() {}
+	virtual ~FVisualLoggerInterface() 
+	{
+
+	}
 
 	virtual bool HasValidCategories(TArray<FVisualLoggerCategoryVerbosityPair> Categories) override
 	{
 		for (const auto& CurrentCategory : Categories)
 		{
-			if (VisualLogger->GetVisualLoggerFilters()->IsFilterEnabled(CurrentCategory.CategoryName.ToString(), CurrentCategory.Verbosity))
+			if (VisualLogger.Pin()->GetVisualLoggerFilters()->IsFilterEnabled(CurrentCategory.CategoryName.ToString(), CurrentCategory.Verbosity))
 			{
 				return true;
 			}
@@ -62,17 +65,17 @@ struct FVisualLoggerInterface : public IVisualLoggerInterface
 
 	bool IsValidCategory(const FString& InCategoryName, TEnumAsByte<ELogVerbosity::Type> Verbosity)
 	{
-		return VisualLogger->GetVisualLoggerFilters()->IsFilterEnabled(InCategoryName, Verbosity);
+		return VisualLogger.Pin()->GetVisualLoggerFilters()->IsFilterEnabled(InCategoryName, Verbosity);
 	}
 
 	bool IsValidCategory(const FString& InGraphName, const FString& InDataName, TEnumAsByte<ELogVerbosity::Type> Verbosity = ELogVerbosity::All)
 	{
-		return VisualLogger->GetVisualLoggerFilters()->IsFilterEnabled(InGraphName, InDataName, Verbosity);
+		return VisualLogger.Pin()->GetVisualLoggerFilters()->IsFilterEnabled(InGraphName, InDataName, Verbosity);
 	}
 
 	virtual FLinearColor GetCategoryColor(FName Category) override
 	{
-		return VisualLogger->GetVisualLoggerFilters()->GetColorForUsedCategory(Category.ToString());
+		return VisualLogger.Pin()->GetVisualLoggerFilters()->GetColorForUsedCategory(Category.ToString());
 	}
 
 	UWorld* GetWorld() const
@@ -116,7 +119,7 @@ struct FVisualLoggerInterface : public IVisualLoggerInterface
 	}
 
 protected:
-	TSharedPtr<SVisualLogger> VisualLogger;
+	TWeakPtr<SVisualLogger> VisualLogger;
 };
 
 SVisualLogger::FVisualLoggerDevice::FVisualLoggerDevice(SVisualLogger* InOwner)
@@ -128,14 +131,14 @@ SVisualLogger::FVisualLoggerDevice::~FVisualLoggerDevice()
 {
 }
 
-void SVisualLogger::FVisualLoggerDevice::Serialize(const class UObject* LogOwner, FName OwnerName, const FVisualLogEntry& LogEntry)
+void SVisualLogger::FVisualLoggerDevice::Serialize(const class UObject* LogOwner, FName OwnerName, FName OwnerClassName, const FVisualLogEntry& LogEntry)
 {
 	if (Owner == NULL)
 	{
 		return;
 	}
 
-	Owner->OnNewLogEntry(FVisualLogDevice::FVisualLogEntryItem(OwnerName, LogEntry));
+	Owner->OnNewLogEntry(FVisualLogDevice::FVisualLogEntryItem(OwnerName, OwnerClassName, LogEntry));
 }
 
 /* SMessagingDebugger constructors
@@ -149,16 +152,59 @@ SVisualLogger::SVisualLogger() : SCompoundWidget(), CommandList(MakeShareable(ne
 
 SVisualLogger::~SVisualLogger()
 {
+	UWorld* World = NULL;
+#if WITH_EDITOR
+	ULogVisualizerSettings* Settings = ULogVisualizerSettings::StaticClass()->GetDefaultObject<ULogVisualizerSettings>();
+	if (Settings->bPresistentFilters)
+	{
+		Settings->PresistentPresets = Settings->CurrentPresets;
+	}
+	else
+	{
+		Settings->PresistentPresets = Settings->CurrentPresets = FFiltersPreset();
+	}
+	Settings->SaveConfig();
 
+
+	UEditorEngine *EEngine = Cast<UEditorEngine>(GEngine);
+	if (GIsEditor && EEngine != NULL)
+	{
+		// lets use PlayWorld during PIE/Simulate and regular world from editor otherwise, to draw debug information
+		World = EEngine->PlayWorld != NULL ? EEngine->PlayWorld : EEngine->GetEditorWorldContext().World();
+
+	}
+	else
+#endif
+	if (!GIsEditor)
+	{
+	World = GEngine->GetWorld();
+	}
+
+	if (World == NULL)
+	{
+		World = GWorld;
+	}
+
+	if (World)
+	{
+		for (TActorIterator<AVisualLoggerRenderingActor> It(World); It; ++It)
+		{
+			World->DestroyActor(*It);
+		}
+	}
+
+	UDebugDrawService::Unregister(FDebugDrawDelegate::CreateRaw(VisualLoggerCanvasRenderer.Get(), &FVisualLoggerCanvasRenderer::DrawOnCanvas));
+	VisualLoggerCanvasRenderer.Reset();
+
+	FVisualLogger::Get().RemoveDevice(InternalDevice.Get());
+	InternalDevice.Reset();
+	
+	VisualLoggerInterface.Reset();
 }
 
 void SVisualLogger::OnTabLosed()
 {
-	UDebugDrawService::Unregister(FDebugDrawDelegate::CreateRaw(VisualLoggerCanvasRenderer.Get(), &FVisualLoggerCanvasRenderer::DrawOnCanvas));
-	VisualLoggerCanvasRenderer = NULL;
 
-	FVisualLogger::Get().RemoveDevice(InternalDevice.Get());
-	InternalDevice = NULL;
 }
 
 /* SMessagingDebugger interface
@@ -182,6 +228,9 @@ void SVisualLogger::Construct(const FArguments& InArgs, const TSharedRef<SDockTa
 	const FVisualLoggerCommands& Commands = FVisualLoggerCommands::Get();
 	FUICommandList& ActionList = *CommandList;
 
+	ULogVisualizerSettings* Settings = ULogVisualizerSettings::StaticClass()->GetDefaultObject<ULogVisualizerSettings>();
+	Settings->CurrentPresets = Settings->bPresistentFilters ? Settings->PresistentPresets : FFiltersPreset();
+
 	ActionList.MapAction(Commands.StartRecording, FExecuteAction::CreateRaw(this, &SVisualLogger::HandleStartRecordingCommandExecute), FCanExecuteAction::CreateRaw(this, &SVisualLogger::HandleStartRecordingCommandCanExecute), FIsActionChecked(), FIsActionButtonVisible::CreateRaw(this, &SVisualLogger::HandleStartRecordingCommandIsVisible));
 	ActionList.MapAction(Commands.StopRecording, FExecuteAction::CreateRaw(this, &SVisualLogger::HandleStopRecordingCommandExecute), FCanExecuteAction::CreateRaw(this, &SVisualLogger::HandleStopRecordingCommandCanExecute), FIsActionChecked(), FIsActionButtonVisible::CreateRaw(this, &SVisualLogger::HandleStopRecordingCommandIsVisible));
 	ActionList.MapAction(Commands.Pause, FExecuteAction::CreateRaw(this, &SVisualLogger::HandlePauseCommandExecute), FCanExecuteAction::CreateRaw(this, &SVisualLogger::HandlePauseCommandCanExecute), FIsActionChecked(), FIsActionButtonVisible::CreateRaw(this, &SVisualLogger::HandlePauseCommandIsVisible));
@@ -283,6 +332,14 @@ void SVisualLogger::Construct(const FArguments& InArgs, const TSharedRef<SDockTa
 
 	// Window Menu
 	FMenuBarBuilder MenuBarBuilder = FMenuBarBuilder(TSharedPtr<FUICommandList>());
+#if 0 //disabled File menu for now (SebaK)
+	MenuBarBuilder.AddPullDownMenu(
+		LOCTEXT("FileMenuLabel", "File"),
+		FText::GetEmpty(),
+		FNewMenuDelegate::CreateRaw(this, &SVisualLogger::FillFileMenu, TabManager),
+		"File"
+		);
+#endif
 	MenuBarBuilder.AddPullDownMenu(
 		LOCTEXT("WindowMenuLabel", "Window"),
 		FText::GetEmpty(),
@@ -339,6 +396,53 @@ void SVisualLogger::HandleTabManagerPersistLayout(const TSharedRef<FTabManager::
 	FLayoutSaveRestore::SaveToConfig(GEditorLayoutIni, LayoutToSave);
 }
 
+void SVisualLogger::FillFileMenu(FMenuBuilder& MenuBuilder, const TSharedPtr<FTabManager> TabManager)
+{
+	MenuBuilder.BeginSection("LogFile", LOCTEXT("FileMenu", "Log File"));
+	{
+		MenuBuilder.AddMenuEntry(FVisualLoggerCommands::Get().Load);
+		MenuBuilder.AddMenuEntry(FVisualLoggerCommands::Get().Save);
+	}
+	MenuBuilder.EndSection();
+	MenuBuilder.BeginSection("LogFilters", LOCTEXT("FIlterMenu", "Log Filters"));
+	{
+		MenuBuilder.AddSubMenu(
+			LOCTEXT("LoadPreset", "Load Preset"), LOCTEXT("LoadPresetTooltip", "Load filter's preset"),
+			FNewMenuDelegate::CreateRaw(this, &SVisualLogger::FillLoadPresetMenu));
+
+		MenuBuilder.AddMenuEntry(
+			LOCTEXT("SavePreset", "Save Preset"), LOCTEXT("SavePresetTooltip", "Save filter's setting as preset"),
+			FSlateIcon(), FUIAction(
+			FExecuteAction::CreateLambda(
+				[this](){
+				}
+			)));
+	}
+	MenuBuilder.EndSection();
+}
+
+void SVisualLogger::FillLoadPresetMenu(FMenuBuilder& MenuBuilder)
+{
+	ULogVisualizerSettings* Settings = ULogVisualizerSettings::StaticClass()->GetDefaultObject<ULogVisualizerSettings>();
+	
+	MenuBuilder.BeginSection("FilterPresets", LOCTEXT("FilterPresets", "Presets"));
+	for (auto& CurrentPreset : Settings->FilterPresets)
+	{
+		MenuBuilder.AddMenuEntry(
+			FText::FromString(CurrentPreset.FilterName), LOCTEXT("LoadPresetTooltip", "Load filter's preset"),
+			FSlateIcon(), FUIAction(
+			FExecuteAction::CreateLambda([this, CurrentPreset](){ this->SetFiltersPreset(CurrentPreset); }
+			))
+		);
+	}
+	MenuBuilder.EndSection();
+}
+
+void SVisualLogger::SetFiltersPreset(const struct FFiltersPreset& Preset)
+{
+	OnFiltersSearchChanged(FText::FromString(Preset.DataFilter));
+}
+
 void SVisualLogger::FillWindowMenu(FMenuBuilder& MenuBuilder, const TSharedPtr<FTabManager> TabManager)
 {
 	if (!TabManager.IsValid())
@@ -685,8 +789,20 @@ void SVisualLogger::OnObjectSelectionChanged(TSharedPtr<class STimeline> TimeLin
 
 void SVisualLogger::OnFiltersSearchChanged(const FText& Filter)
 {
-	VisualLoggerFilters->OnFiltersSearchChanged(Filter);
-	LogsList->OnFiltersSearchChanged(Filter);
-	MainView->OnFiltersSearchChanged(Filter);
+	ULogVisualizerSettings::StaticClass()->GetDefaultObject<ULogVisualizerSettings>()->CurrentPresets.DataFilter = Filter.ToString();
+	ULogVisualizerSettings::StaticClass()->GetDefaultObject<ULogVisualizerSettings>()->SaveConfig();
+
+	if (VisualLoggerFilters.IsValid())
+	{
+		VisualLoggerFilters->OnFiltersSearchChanged(Filter);
+	}
+	if (LogsList.IsValid())
+	{
+		LogsList->OnFiltersSearchChanged(Filter);
+	}
+	if (MainView.IsValid())
+	{
+		MainView->OnFiltersSearchChanged(Filter);
+	}
 }
 #undef LOCTEXT_NAMESPACE
diff --git a/Engine/Source/Developer/LogVisualizer/Private/SVisualLogger.h b/Engine/Source/Developer/LogVisualizer/Private/SVisualLogger.h
--- a/Engine/Source/Developer/LogVisualizer/Private/SVisualLogger.h
+++ b/Engine/Source/Developer/LogVisualizer/Private/SVisualLogger.h
@@ -40,7 +40,7 @@ public:
 	public:
 		FVisualLoggerDevice(SVisualLogger* Owner);
 		virtual ~FVisualLoggerDevice();
-		virtual void Serialize(const class UObject* LogOwner, FName OwnerName, const FVisualLogEntry& LogEntry) override;
+		virtual void Serialize(const class UObject* LogOwner, FName OwnerName, FName OwnerClassName, const FVisualLogEntry& LogEntry) override;
 
 	protected:
 		SVisualLogger* Owner;
@@ -48,7 +48,9 @@ public:
 	friend class FVisualLogDevice;
 
 	TSharedRef<SDockTab> HandleTabManagerSpawnTab(const FSpawnTabArgs& Args, FName TabIdentifier) const;
+	void FillFileMenu(FMenuBuilder& MenuBuilder, const TSharedPtr<FTabManager> TabManager);
 	static void FillWindowMenu(FMenuBuilder& MenuBuilder, const TSharedPtr<FTabManager> TabManager);
+	void FillLoadPresetMenu(FMenuBuilder& Builder);
 
 	/** Callback for for when the owner tab's visual state is being persisted. */
 	void HandleMajorTabPersistVisualState();
@@ -89,6 +91,9 @@ public:
 
 	TSharedPtr<SVisualLoggerFilters> GetVisualLoggerFilters() { return VisualLoggerFilters; }
 	void HandleTabManagerPersistLayout(const TSharedRef<FTabManager::FLayout>& LayoutToSave);
+
+	void SetFiltersPreset(const struct FFiltersPreset& Preset);
+
 protected:
 	// Holds the list of UI commands.
 	TSharedRef<FUICommandList> CommandList;
@@ -102,6 +107,7 @@ protected:
 	TWeakObjectPtr<class AVisualLoggerCameraController> CameraController;
 	TSharedPtr<struct FVisualLoggerCanvasRenderer> VisualLoggerCanvasRenderer;
 	TSharedPtr<IVisualLoggerInterface> VisualLoggerInterface;
+
 	mutable TSharedPtr<SVisualLoggerFilters> VisualLoggerFilters;
 	mutable TSharedPtr<SVisualLoggerView> MainView;
 	mutable TSharedPtr<SVisualLoggerLogsList> LogsList;
diff --git a/Engine/Source/Developer/LogVisualizer/Private/SVisualLoggerFilters.cpp b/Engine/Source/Developer/LogVisualizer/Private/SVisualLoggerFilters.cpp
--- a/Engine/Source/Developer/LogVisualizer/Private/SVisualLoggerFilters.cpp
+++ b/Engine/Source/Developer/LogVisualizer/Private/SVisualLoggerFilters.cpp
@@ -367,13 +367,17 @@ void SVisualLoggerFilters::AddFilter(const FString& InFilterName)
 		SNew(SFilterWidget)
 		.FilterName(*InFilterName)
 		.ColorCategory(Color)
-		.OnFilterChanged(VisualLoggerInterface->GetVisualLoggerEvents().OnFiltersChanged)
+		.OnFilterChanged(SFilterWidget::FOnSimpleRequest::CreateRaw(this, &SVisualLoggerFilters::OnFiltersChanged))
 		//.OnRequestRemove(this, &SLogFilterList::RemoveFilter)
 		//.OnRequestDisableAll(this, &SLogFilterList::DisableAllFilters)
 		//.OnRequestRemoveAll(this, &SLogFilterList::RemoveAllFilters);
 		;
-	NewFilter->SetEnabled(true);
+
+	ULogVisualizerSettings* Settings = ULogVisualizerSettings::StaticClass()->GetDefaultObject<ULogVisualizerSettings>();
+	TArray<FString> CurrentFilters = Settings->CurrentPresets.CategoryFilters;
+	NewFilter->SetEnabled(CurrentFilters.Num() == 0 || CurrentFilters.Find(InFilterName) != INDEX_NONE);
 	Filters.Add(NewFilter);
+	Settings->CurrentPresets.CategoryFilters = CurrentFilters;
 
 	FilterBox->AddSlot()
 		.Padding(2, 2)
@@ -384,6 +388,23 @@ void SVisualLoggerFilters::AddFilter(const FString& InFilterName)
 	GraphsFilterCombo->SetVisibility(GraphFilters.Num() > 0 ? EVisibility::Visible : EVisibility::Collapsed);
 }
 
+void SVisualLoggerFilters::OnFiltersChanged()
+{
+	TArray<FString> EnabledFilters;
+	for (TSharedRef<SFilterWidget> CurrentFilter : Filters)
+	{
+		if (CurrentFilter->IsEnabled())
+		{
+			EnabledFilters.AddUnique(CurrentFilter->GetFilterName().ToString());
+		}
+	}
+	ULogVisualizerSettings* Settings = ULogVisualizerSettings::StaticClass()->GetDefaultObject<ULogVisualizerSettings>();
+	Settings->CurrentPresets.CategoryFilters = EnabledFilters;
+	Settings->SaveConfig();
+
+	VisualLoggerInterface.Pin()->GetVisualLoggerEvents().OnFiltersChanged.ExecuteIfBound();
+}
+
 void SVisualLoggerFilters::AddFilter(const FString& GraphName, const FString& DataName)
 {
 	TArray<FSimpleGraphFilter>& DataNames = GraphFilters.FindOrAdd(*GraphName);
@@ -396,6 +417,7 @@ void SVisualLoggerFilters::AddFilter(const FString& GraphName, const FString& Da
 bool SVisualLoggerFilters::IsFilterEnabled(const FString& InFilterName, TEnumAsByte<ELogVerbosity::Type> Verbosity)
 {
 	const ULogVisualizerSettings* Settings = ULogVisualizerSettings::StaticClass()->GetDefaultObject<ULogVisualizerSettings>();
+	FString FiltersSearchString = Settings->CurrentPresets.DataFilter;
 	const FName FilterName(*InFilterName);
 	for (int32 Index = 0; Index < Filters.Num(); ++Index)
 	{
@@ -427,7 +449,7 @@ bool SVisualLoggerFilters::IsFilterEnabled(const FString& InGraphName, const FSt
 
 void SVisualLoggerFilters::OnFiltersSearchChanged(const FText& Filter)
 {
-	FiltersSearchString = Filter.ToString();
+
 }
 
 #undef LOCTEXT_NAMESPACE
diff --git a/Engine/Source/Developer/LogVisualizer/Private/SVisualLoggerFilters.h b/Engine/Source/Developer/LogVisualizer/Private/SVisualLoggerFilters.h
--- a/Engine/Source/Developer/LogVisualizer/Private/SVisualLoggerFilters.h
+++ b/Engine/Source/Developer/LogVisualizer/Private/SVisualLoggerFilters.h
@@ -43,6 +43,7 @@ public:
 	void InvalidateCanvas();
 	void GraphFilterCategoryClicked(FName MenuCategory);
 	bool IsGraphFilterCategoryInUse(FName MenuCategory) const;
+	void OnFiltersChanged();
 
 	TSharedRef<SWidget> MakeGraphsFilterMenu();
 	void CreateFiltersMenuCategoryForGraph(FMenuBuilder& MenuBuilder, FName MenuCategory) const;
@@ -50,10 +51,9 @@ public:
 protected:
 	/** The horizontal box which contains all the filters */
 	TSharedPtr<SWrapBox> FilterBox;
-	TSharedPtr<IVisualLoggerInterface> VisualLoggerInterface;
+	TWeakPtr<IVisualLoggerInterface> VisualLoggerInterface;
 	TArray<TSharedRef<SFilterWidget> > Filters;
 	TMap<FName, TArray<FSimpleGraphFilter> > GraphFilters;
-	FString FiltersSearchString;
 	TSharedPtr<SComboButton> GraphsFilterCombo;
 	FString GraphsFilter;
 	static FColor ColorPalette[];
diff --git a/Engine/Source/Developer/LogVisualizer/Private/SVisualLoggerLogsList.cpp b/Engine/Source/Developer/LogVisualizer/Private/SVisualLoggerLogsList.cpp
--- a/Engine/Source/Developer/LogVisualizer/Private/SVisualLoggerLogsList.cpp
+++ b/Engine/Source/Developer/LogVisualizer/Private/SVisualLoggerLogsList.cpp
@@ -102,7 +102,7 @@ TSharedRef<ITableRow> SVisualLoggerLogsList::LogEntryLinesGenerateRow(TSharedPtr
 			[
 				SNew(STextBlock)
 				.ColorAndOpacity(FSlateColor(Item->CategoryColor))
-				.Text(Item->Category)
+				.Text( FText::FromString(Item->Category) )
 				.HighlightText(this, &SVisualLoggerLogsList::GetFilterText)
 			]
 			+ SHorizontalBox::Slot()
@@ -111,14 +111,14 @@ TSharedRef<ITableRow> SVisualLoggerLogsList::LogEntryLinesGenerateRow(TSharedPtr
 			[
 				SNew(STextBlock)
 				.ColorAndOpacity(FSlateColor(FLinearColor::Gray))
-				.Text(FString(TEXT("(")) + FString(FOutputDevice::VerbosityToString(Item->Verbosity)) + FString(TEXT(")")))
+				.Text(FText::FromString(FString(TEXT("(")) + FString(FOutputDevice::VerbosityToString(Item->Verbosity)) + FString(TEXT(")"))))
 			]
 			+ SHorizontalBox::Slot()
 			.Padding(FMargin(5.0f, 0.0f))
 			[
 				SNew(STextBlock)
 				.AutoWrapText(true)
-				.Text(Item->Line)
+				.Text(FText::FromString(Item->Line))
 				.HighlightText(this, &SVisualLoggerLogsList::GetFilterText)
 			]
 		];
@@ -128,12 +128,11 @@ FText SVisualLoggerLogsList::GetFilterText() const
 {
 	static FText NoText;
 	const bool bSearchInsideLogs = ULogVisualizerSettings::StaticClass()->GetDefaultObject<ULogVisualizerSettings>()->bSearchInsideLogs;
-	return bSearchInsideLogs ? FilterText : NoText;
+	return bSearchInsideLogs ? FText::FromString(ULogVisualizerSettings::StaticClass()->GetDefaultObject<ULogVisualizerSettings>()->CurrentPresets.DataFilter) : NoText;
 }
 
 void SVisualLoggerLogsList::OnFiltersSearchChanged(const FText& Filter)
 {
-	FilterText = Filter;
 	OnFiltersChanged();
 }
 
@@ -144,7 +143,7 @@ void SVisualLoggerLogsList::OnItemSelectionChanged(const FVisualLogDevice::FVisu
 
 	TArray<FVisualLoggerCategoryVerbosityPair> OutCategories;
 	FVisualLoggerHelpers::GetCategories(LogEntry.Entry, OutCategories);
-	if (VisualLoggerInterface->HasValidCategories(OutCategories) == false)
+	if (VisualLoggerInterface.Pin()->HasValidCategories(OutCategories) == false)
 	{
 		return;
 	}
@@ -156,13 +155,13 @@ void SVisualLoggerLogsList::OnItemSelectionChanged(const FVisualLogDevice::FVisu
 		bool bShowLine = true;
 
 		FString CurrentCategory = LogLine->Category.ToString();
-		bShowLine = VisualLoggerInterface->IsValidCategory(CurrentCategory, LogLine->Verbosity) /*&& (bHistogramGraphsFilter || (QuickFilterText.Len() == 0 || CurrentCategory.Find(QuickFilterText) != INDEX_NONE))*/;
+		bShowLine = VisualLoggerInterface.Pin()->IsValidCategory(CurrentCategory, LogLine->Verbosity) /*&& (bHistogramGraphsFilter || (QuickFilterText.Len() == 0 || CurrentCategory.Find(QuickFilterText) != INDEX_NONE))*/;
 
 		if (bShowLine)
 		{
 			FLogEntryItem EntryItem;
 			EntryItem.Category = LogLine->Category.ToString();
-			EntryItem.CategoryColor = VisualLoggerInterface->GetCategoryColor(LogLine->Category);
+			EntryItem.CategoryColor = VisualLoggerInterface.Pin()->GetCategoryColor(LogLine->Category);
 			EntryItem.Verbosity = LogLine->Verbosity;
 			EntryItem.Line = LogLine->Line;
 			EntryItem.UserData = LogLine->UserData;
@@ -177,13 +176,13 @@ void SVisualLoggerLogsList::OnItemSelectionChanged(const FVisualLogDevice::FVisu
 		bool bShowLine = true;
 
 		FString CurrentCategory = Event.Name;
-		bShowLine = VisualLoggerInterface->IsValidCategory(Event.Name, Event.Verbosity) /*&& (bHistogramGraphsFilter || (QuickFilterText.Len() == 0 || CurrentCategory.Find(QuickFilterText) != INDEX_NONE))*/;
+		bShowLine = VisualLoggerInterface.Pin()->IsValidCategory(Event.Name, Event.Verbosity) /*&& (bHistogramGraphsFilter || (QuickFilterText.Len() == 0 || CurrentCategory.Find(QuickFilterText) != INDEX_NONE))*/;
 
 		if (bShowLine)
 		{
 			FLogEntryItem EntryItem;
 			EntryItem.Category = Event.Name;
-			EntryItem.CategoryColor = VisualLoggerInterface->GetCategoryColor(*EntryItem.Category);
+			EntryItem.CategoryColor = VisualLoggerInterface.Pin()->GetCategoryColor(*EntryItem.Category);
 			EntryItem.Verbosity = Event.Verbosity;
 			EntryItem.Line = FString::Printf(TEXT("Registered event: '%s' (%d times)%s"), *Event.Name, Event.Counter, Event.EventTags.Num() ? TEXT("\n") : TEXT(""));
 			for (auto& EventTag : Event.EventTags)
diff --git a/Engine/Source/Developer/LogVisualizer/Private/SVisualLoggerLogsList.h b/Engine/Source/Developer/LogVisualizer/Private/SVisualLoggerLogsList.h
--- a/Engine/Source/Developer/LogVisualizer/Private/SVisualLoggerLogsList.h
+++ b/Engine/Source/Developer/LogVisualizer/Private/SVisualLoggerLogsList.h
@@ -23,7 +23,6 @@ public:
 protected:
 	TSharedPtr<SListView<TSharedPtr<struct FLogEntryItem> > > LogsLinesWidget;
 	TArray<TSharedPtr<struct FLogEntryItem> > LogEntryLines;
-	TSharedPtr<IVisualLoggerInterface> VisualLoggerInterface;
+	TWeakPtr<IVisualLoggerInterface> VisualLoggerInterface;
 	FVisualLogDevice::FVisualLogEntryItem CurrentLogEntry;
-	FText FilterText;
 };
diff --git a/Engine/Source/Developer/LogVisualizer/Private/SVisualLoggerStatusView.cpp b/Engine/Source/Developer/LogVisualizer/Private/SVisualLoggerStatusView.cpp
--- a/Engine/Source/Developer/LogVisualizer/Private/SVisualLoggerStatusView.cpp
+++ b/Engine/Source/Developer/LogVisualizer/Private/SVisualLoggerStatusView.cpp
@@ -107,7 +107,7 @@ TSharedRef<ITableRow> SVisualLoggerStatusView::HandleGenerateLogStatus(TSharedPt
 		return SNew(STableRow<TSharedPtr<FLogStatusItem> >, OwnerTable)
 			[
 				SNew(STextBlock)
-				.Text(InItem->ItemText)
+				.Text(FText::FromString(InItem->ItemText))
 			];
 	}
 
@@ -116,21 +116,21 @@ TSharedRef<ITableRow> SVisualLoggerStatusView::HandleGenerateLogStatus(TSharedPt
 		[
 			SNew(SBorder)
 			.BorderImage(FLogVisualizerStyle::Get().GetBrush("NoBorder"))
-			.ToolTipText(TooltipText)
+			.ToolTipText(FText::FromString(TooltipText))
 			[
 				SNew(SHorizontalBox)
 				+ SHorizontalBox::Slot()
 				.AutoWidth()
 				[
 					SNew(STextBlock)
-					.Text(InItem->ItemText)
+					.Text(FText::FromString(InItem->ItemText))
 					.ColorAndOpacity(FColorList::Aquamarine)
 				]
 				+ SHorizontalBox::Slot()
 				.Padding(4.0f, 0, 0, 0)
 				[
 					SNew(STextBlock)
-					.Text(InItem->ValueText)
+					.Text(FText::FromString(InItem->ValueText))
 				]
 			]
 		];
diff --git a/Engine/Source/Developer/LogVisualizer/Private/SVisualLoggerView.cpp b/Engine/Source/Developer/LogVisualizer/Private/SVisualLoggerView.cpp
--- a/Engine/Source/Developer/LogVisualizer/Private/SVisualLoggerView.cpp
+++ b/Engine/Source/Developer/LogVisualizer/Private/SVisualLoggerView.cpp
@@ -73,7 +73,7 @@ void SVisualLoggerView::Construct(const FArguments& InArgs, const TSharedRef<FUI
 {
 	VisualLoggerInterface = InVisualLoggerInterface;
 	AnimationOutlinerFillPercentage = .25f;
-	VisualLoggerEvents = VisualLoggerInterface->GetVisualLoggerEvents();
+	VisualLoggerEvents = VisualLoggerInterface.Pin()->GetVisualLoggerEvents();
 
 	FVisualLoggerTimeSliderArgs TimeSliderArgs;
 	TimeSliderArgs.ViewRange = InArgs._ViewRange;
@@ -89,7 +89,7 @@ void SVisualLoggerView::Construct(const FArguments& InArgs, const TSharedRef<FUI
 
 	TSharedPtr<FSequencerTimeSliderController> TimeSliderController(new FSequencerTimeSliderController(TimeSliderArgs));
 	TimeSliderController->SetExternalScrollbar(ZoomScrollBar);
-	VisualLoggerInterface->SetTimeSliderController(TimeSliderController);
+	VisualLoggerInterface.Pin()->SetTimeSliderController(TimeSliderController);
 
 	// Create the top and bottom sliders
 	const bool bMirrorLabels = true;
@@ -102,6 +102,7 @@ void SVisualLoggerView::Construct(const FArguments& InArgs, const TSharedRef<FUI
 
 
 	ULogVisualizerSettings* Settings = ULogVisualizerSettings::StaticClass()->GetDefaultObject<ULogVisualizerSettings>();
+	FFiltersPreset& CurrentPresets = Settings->CurrentPresets;
 
 	ChildSlot
 		[
@@ -126,7 +127,39 @@ void SVisualLoggerView::Construct(const FArguments& InArgs, const TSharedRef<FUI
 							SNew(SHorizontalBox)
 							+ SHorizontalBox::Slot()
 							.Padding(FMargin(0))
+							.AutoWidth()
+							.HAlign(HAlign_Center)
 							.VAlign(VAlign_Center)
+							[
+								SNew(SImage)
+								.Visibility_Lambda([this]()->EVisibility{ return ULogVisualizerSettings::StaticClass()->GetDefaultObject<ULogVisualizerSettings>()->CurrentPresets.SelectedClasses.Num() > 0 ? EVisibility::Visible : EVisibility::Collapsed; })
+								.Image(FLogVisualizerStyle::Get().GetBrush("Filters.FilterIcon"))
+							]
+							+ SHorizontalBox::Slot()
+							.Padding(FMargin(0))
+							.HAlign(HAlign_Right)
+							.AutoWidth()
+							[
+								SAssignNew(ClassesComboButton, SComboButton)
+								.Visibility_Lambda([this]()->EVisibility{ return TimelinesContainer.IsValid() && TimelinesContainer->GetAllNodes().Num() > 1 ? EVisibility::Visible : EVisibility::Collapsed; })
+								.ComboButtonStyle(FLogVisualizerStyle::Get(), "Filters.Style")
+								.ForegroundColor(FLinearColor::White)
+								.ContentPadding(0)
+								.OnGetMenuContent(this, &SVisualLoggerView::MakeClassesFilterMenu)
+								.ToolTipText(LOCTEXT("SetFilterByClasses", "Select classes to show"))
+								.HasDownArrow(true)
+								.ContentPadding(FMargin(1, 0))
+								.ButtonContent()
+								[
+									SNew(STextBlock)
+									.TextStyle(FLogVisualizerStyle::Get(), "Filters.Text")
+									.Text(LOCTEXT("FilterClasses", "Classes"))
+								]
+							]
+							+ SHorizontalBox::Slot()
+							.Padding(FMargin(0))
+							.HAlign(HAlign_Fill)
+							.FillWidth(1)
 							[
 								SNew(SBox)
 								.Padding(FMargin(0, 0, 4, 0))
@@ -237,6 +270,7 @@ void SVisualLoggerView::Construct(const FArguments& InArgs, const TSharedRef<FUI
 			]
 		];
 
+		SearchBox->SetText(FText::FromString(CurrentPresets.DataFilter));
 		//ScrollBox->AddSlot()
 		//[
 		//	SAssignNew(TimelinesContainer, STimelinesContainer, SharedThis(this), TimeSliderController.ToSharedRef())
@@ -319,12 +353,81 @@ void SVisualLoggerView::OnFiltersSearchChanged(const FText& Filter)
 
 FCursorReply SVisualLoggerView::OnCursorQuery(const FGeometry& MyGeometry, const FPointerEvent& CursorEvent) const
 {
-	if (VisualLoggerInterface->GetTimeSliderController()->IsPanning())
+	if (VisualLoggerInterface.Pin()->GetTimeSliderController()->IsPanning())
 	{
 		return FCursorReply::Cursor(EMouseCursor::GrabHand);
 	}
 	return FCursorReply::Cursor(EMouseCursor::Default);
 }
 
+TSharedRef<SWidget> SVisualLoggerView::MakeClassesFilterMenu()
+{
+	const TArray< TSharedPtr<class STimeline> >& AllTimelines = TimelinesContainer->GetAllNodes();
+
+	FMenuBuilder MenuBuilder(true, NULL);
+
+	TArray<FString> UniqueClasses;
+	MenuBuilder.BeginSection(TEXT("Graphs"));
+	for (TSharedPtr<class STimeline> CurrentTimeline : AllTimelines)
+	{
+		FString OwnerClassName = CurrentTimeline->GetOwnerClassName().ToString();
+		if (UniqueClasses.Find(OwnerClassName) == INDEX_NONE)
+		{
+			FText LabelText = FText::FromString(OwnerClassName);
+			MenuBuilder.AddMenuEntry(
+				LabelText,
+				FText::Format(LOCTEXT("FilterByClassPrefix", "Toggle {0} class"), LabelText),
+				FSlateIcon(),
+				FUIAction(
+				FExecuteAction::CreateLambda([this, OwnerClassName]()
+				{
+					TArray<FString>& SelectedOwnerClasses = ULogVisualizerSettings::StaticClass()->GetDefaultObject<ULogVisualizerSettings>()->CurrentPresets.SelectedClasses;
+					if (SelectedOwnerClasses.Find(OwnerClassName) != INDEX_NONE)
+						SelectedOwnerClasses.Remove(OwnerClassName);
+					else
+						SelectedOwnerClasses.AddUnique(OwnerClassName);
+
+					OnChangedClassesFilter();
+				}),
+				FCanExecuteAction(),
+				FIsActionChecked::CreateLambda([OwnerClassName]()->bool
+				{
+					TArray<FString>& SelectedOwnerClasses = ULogVisualizerSettings::StaticClass()->GetDefaultObject<ULogVisualizerSettings>()->CurrentPresets.SelectedClasses;
+					return SelectedOwnerClasses.Find(OwnerClassName) != INDEX_NONE;
+				}),
+				FIsActionButtonVisible()),
+				NAME_None,
+				EUserInterfaceActionType::ToggleButton
+				);
+			UniqueClasses.AddUnique(OwnerClassName);
+		}
+	}
+	MenuBuilder.EndSection(); //ContentBrowserFilterBasicAsset
+
+
+	FDisplayMetrics DisplayMetrics;
+	FSlateApplication::Get().GetDisplayMetrics(DisplayMetrics);
+
+	const FVector2D DisplaySize(
+		DisplayMetrics.PrimaryDisplayWorkAreaRect.Right - DisplayMetrics.PrimaryDisplayWorkAreaRect.Left,
+		DisplayMetrics.PrimaryDisplayWorkAreaRect.Bottom - DisplayMetrics.PrimaryDisplayWorkAreaRect.Top);
+
+	return
+		SNew(SVerticalBox)
+		+ SVerticalBox::Slot()
+		.MaxHeight(DisplaySize.Y * 0.5)
+		[
+			MenuBuilder.MakeWidget()
+		];
+}
+
+void SVisualLoggerView::OnChangedClassesFilter()
+{
+	ULogVisualizerSettings::StaticClass()->GetDefaultObject<ULogVisualizerSettings>()->SaveConfig();
+	for (auto CurrentItem : TimelinesContainer->GetAllNodes())
+	{
+		CurrentItem->UpdateVisibility();
+	}
+}
 
 #undef LOCTEXT_NAMESPACE
diff --git a/Engine/Source/Developer/LogVisualizer/Private/SVisualLoggerView.h b/Engine/Source/Developer/LogVisualizer/Private/SVisualLoggerView.h
--- a/Engine/Source/Developer/LogVisualizer/Private/SVisualLoggerView.h
+++ b/Engine/Source/Developer/LogVisualizer/Private/SVisualLoggerView.h
@@ -34,16 +34,19 @@ public:
 	void OnFiltersSearchChanged(const FText& Filter);
 	void OnSearchSplitterResized();
 	void OnObjectSelectionChanged(TSharedPtr<class STimeline> TimeLine);
+	void OnChangedClassesFilter();
 
 	void GetTimelines(TArray<TSharedPtr<class STimeline> >&, bool bOnlySelectedOnes = false);
 	virtual FCursorReply OnCursorQuery(const FGeometry& MyGeometry, const FPointerEvent& CursorEvent) const override;
+	TSharedRef<SWidget> MakeClassesFilterMenu();
 
 protected:
-	TSharedPtr<struct IVisualLoggerInterface> VisualLoggerInterface;
+	TWeakPtr<struct IVisualLoggerInterface> VisualLoggerInterface;
 	TSharedPtr<class STimelinesContainer> TimelinesContainer;
 	TSharedPtr<class SSplitter> SearchSplitter;
 	TSharedPtr<class SScrollBox> ScrollBox;
 	TSharedPtr<class SSearchBox> SearchBox;
+	TSharedPtr<class SComboButton> ClassesComboButton;
 
 	FVisualLoggerEvents	VisualLoggerEvents;
 	float AnimationOutlinerFillPercentage;
diff --git a/Engine/Source/Developer/LogVisualizer/Private/VisualLoggerCanvasRenderer.cpp b/Engine/Source/Developer/LogVisualizer/Private/VisualLoggerCanvasRenderer.cpp
--- a/Engine/Source/Developer/LogVisualizer/Private/VisualLoggerCanvasRenderer.cpp
+++ b/Engine/Source/Developer/LogVisualizer/Private/VisualLoggerCanvasRenderer.cpp
@@ -69,7 +69,7 @@ void FVisualLoggerCanvasRenderer::DrawOnCanvas(class UCanvas* Canvas, class APla
 		return;
 	}
 
-	UWorld* World = VisualLoggerInterface->GetWorld();
+	UWorld* World = VisualLoggerInterface.Pin()->GetWorld();
 	if (World == NULL)
 	{
 		return;
@@ -84,7 +84,7 @@ void FVisualLoggerCanvasRenderer::DrawOnCanvas(class UCanvas* Canvas, class APla
 	for (const auto CurrentData : SelectedEntry.DataBlocks)
 	{
 		const FName TagName = CurrentData.TagName;
-		const bool bIsValidByFilter = VisualLoggerInterface->IsValidCategory(CurrentData.Category.ToString(), ELogVerbosity::All) && VisualLoggerInterface->IsValidCategory(CurrentData.TagName.ToString(), ELogVerbosity::All);
+		const bool bIsValidByFilter = VisualLoggerInterface.Pin()->IsValidCategory(CurrentData.Category.ToString(), ELogVerbosity::All) && VisualLoggerInterface.Pin()->IsValidCategory(CurrentData.TagName.ToString(), ELogVerbosity::All);
 		FVisualLogExtensionInterface* Extension = FVisualLogger::Get().GetExtensionForTag(TagName);
 		if (!Extension)
 		{
@@ -93,11 +93,11 @@ void FVisualLoggerCanvasRenderer::DrawOnCanvas(class UCanvas* Canvas, class APla
 
 		if (!bIsValidByFilter)
 		{
-			Extension->DisableDrawingForData(VisualLoggerInterface->GetWorld(), Canvas, NULL, TagName, CurrentData, SelectedEntry.TimeStamp);
+			Extension->DisableDrawingForData(VisualLoggerInterface.Pin()->GetWorld(), Canvas, NULL, TagName, CurrentData, SelectedEntry.TimeStamp);
 		}
 		else
 		{
-			Extension->DrawData(VisualLoggerInterface->GetWorld(), Canvas, NULL, TagName, CurrentData, SelectedEntry.TimeStamp);
+			Extension->DrawData(VisualLoggerInterface.Pin()->GetWorld(), Canvas, NULL, TagName, CurrentData, SelectedEntry.TimeStamp);
 		}
 	}
 
@@ -124,14 +124,14 @@ void FVisualLoggerCanvasRenderer::DrawHistogramGraphs(class UCanvas* Canvas, cla
 		TMap<FName, FGraphLineData> GraphLines;
 	};
 
-	if (VisualLoggerInterface->GetTimeSliderController() == NULL)
+	if (VisualLoggerInterface.Pin()->GetTimeSliderController() == NULL)
 	{
 		return;
 	}
 
 	TMap<FName, FGraphData>	CollectedGraphs;
 
-	const FVisualLoggerTimeSliderArgs&  TimeSliderArgs = VisualLoggerInterface->GetTimeSliderController()->GetTimeSliderArgs();
+	const FVisualLoggerTimeSliderArgs&  TimeSliderArgs = VisualLoggerInterface.Pin()->GetTimeSliderController()->GetTimeSliderArgs();
 	TRange<float> LocalViewRange = TimeSliderArgs.ViewRange.Get();
 	const float LocalViewRangeMin = LocalViewRange.GetLowerBoundValue();
 	const float LocalViewRangeMax = LocalViewRange.GetUpperBoundValue();
@@ -163,7 +163,7 @@ void FVisualLoggerCanvasRenderer::DrawHistogramGraphs(class UCanvas* Canvas, cla
 			const FName CurrentGraphName = CurrentSample.GraphName;
 			const FName CurrentDataName = CurrentSample.DataName;
 
-			const bool bIsValidByFilter = VisualLoggerInterface->IsValidCategory(CurrentSample.GraphName.ToString(), CurrentSample.DataName.ToString(), ELogVerbosity::All);
+			const bool bIsValidByFilter = VisualLoggerInterface.Pin()->IsValidCategory(CurrentSample.GraphName.ToString(), CurrentSample.DataName.ToString(), ELogVerbosity::All);
 			const bool bCurrentDataNamePassed = true;// !bHistogramGraphsFilter || (QuickFilterText.Len() == 0 || CurrentDataName.ToString().Find(QuickFilterText) != INDEX_NONE);
 
 
diff --git a/Engine/Source/Developer/LogVisualizer/Private/VisualLoggerCanvasRenderer.h b/Engine/Source/Developer/LogVisualizer/Private/VisualLoggerCanvasRenderer.h
--- a/Engine/Source/Developer/LogVisualizer/Private/VisualLoggerCanvasRenderer.h
+++ b/Engine/Source/Developer/LogVisualizer/Private/VisualLoggerCanvasRenderer.h
@@ -16,7 +16,7 @@ protected:
 
 private:
 	FVisualLogEntry SelectedEntry;
-	TSharedPtr<IVisualLoggerInterface> VisualLoggerInterface;
+	TWeakPtr<IVisualLoggerInterface> VisualLoggerInterface;
 	TMap<FString, TArray<FString> > UsedGraphCategories;
 	TWeakPtr<class STimeline> CurrentTimeLine;
 };
diff --git a/Engine/Source/Developer/LogVisualizer/Public/LogVisualizerModule.h b/Engine/Source/Developer/LogVisualizer/Public/LogVisualizerModule.h
--- a/Engine/Source/Developer/LogVisualizer/Public/LogVisualizerModule.h
+++ b/Engine/Source/Developer/LogVisualizer/Public/LogVisualizerModule.h
@@ -8,7 +8,7 @@
 class SDockTab;
 class ISlateStyle;
 class FSpawnTabArgs;
-class FNewLogVisualizerModule : public IModuleInterface, public IModularFeature
+class FLogVisualizerModule : public IModuleInterface, public IModularFeature
 {
 public:
 	// Begin IModuleInterface
@@ -18,5 +18,4 @@ public:
 
 private:
 	TSharedRef<SDockTab> SpawnLogVisualizerTab(const FSpawnTabArgs& SpawnTabArgs);
-	void OnTabClosed(TSharedRef<SDockTab>);
 };
diff --git a/Engine/Source/Developer/LogVisualizer/Public/LogVisualizerSettings.h b/Engine/Source/Developer/LogVisualizer/Public/LogVisualizerSettings.h
--- a/Engine/Source/Developer/LogVisualizer/Public/LogVisualizerSettings.h
+++ b/Engine/Source/Developer/LogVisualizer/Public/LogVisualizerSettings.h
@@ -8,6 +8,27 @@ GameplayDebuggerSettings.h: Declares the UGameplayDebuggerSettings class.
 
 #include "LogVisualizerSettings.generated.h"
 
+USTRUCT()
+struct FFiltersPreset
+{
+	GENERATED_USTRUCT_BODY()
+
+	UPROPERTY(config)
+	FString FilterName;
+
+	UPROPERTY(config)
+	FString ObjectNameFilter;
+
+	UPROPERTY(config)
+	FString DataFilter;
+
+	UPROPERTY(config)
+	TArray<FString> CategoryFilters;
+
+	UPROPERTY(config)
+	TArray<FString> SelectedClasses;
+};
+
 UCLASS(config = EditorUserSettings)
 class LOGVISUALIZER_API ULogVisualizerSettings : public UObject
 {
@@ -44,6 +65,19 @@ public:
 	UPROPERTY(EditAnywhere, config, Category = "VisualLogger")
 	FColor GraphsBackgroundColor;
 
+	UPROPERTY(config)
+	TArray<FFiltersPreset> FilterPresets;
+
+
+	/**Whether to store all filter settings on exit*/
+	UPROPERTY(EditAnywhere, config, Category = "VisualLogger")
+	bool bPresistentFilters;
+
+	UPROPERTY(config)
+	FFiltersPreset PresistentPresets;
+
+	FFiltersPreset CurrentPresets;
+
 	// UObject overrides
 #if WITH_EDITOR
 	virtual void PostEditChangeProperty(struct FPropertyChangedEvent& PropertyChangedEvent) override;
