Fix some 2014 copyright notices that have crept in.

diff --git a/Engine/Build/Android/Java/src/com/epicgames/ue4/MessageBox01.java b/Engine/Build/Android/Java/src/com/epicgames/ue4/MessageBox01.java
--- a/Engine/Build/Android/Java/src/com/epicgames/ue4/MessageBox01.java
+++ b/Engine/Build/Android/Java/src/com/epicgames/ue4/MessageBox01.java
@@ -1,4 +1,4 @@
-// Copyright 1998-2014 Epic Games, Inc. All Rights Reserved.
+// Copyright 1998-2015 Epic Games, Inc. All Rights Reserved.
 
 package com.epicgames.ue4;
 
diff --git a/Engine/Plugins/2D/Paper2D/Source/Paper2D/Private/PaperSprite.cpp b/Engine/Plugins/2D/Paper2D/Source/Paper2D/Private/PaperSprite.cpp
--- a/Engine/Plugins/2D/Paper2D/Source/Paper2D/Private/PaperSprite.cpp
+++ b/Engine/Plugins/2D/Paper2D/Source/Paper2D/Private/PaperSprite.cpp
@@ -1,4 +1,4 @@
-// Copyright 1998-2014 Epic Games, Inc. All Rights Reserved.
+// Copyright 1998-2015 Epic Games, Inc. All Rights Reserved.
 
 #include "Paper2DPrivatePCH.h"
 #include "PaperSprite.h"
diff --git a/Engine/Source/Developer/LogVisualizer/Private/LogVisualizerSessionSettings.cpp b/Engine/Source/Developer/LogVisualizer/Private/LogVisualizerSessionSettings.cpp
--- a/Engine/Source/Developer/LogVisualizer/Private/LogVisualizerSessionSettings.cpp
+++ b/Engine/Source/Developer/LogVisualizer/Private/LogVisualizerSessionSettings.cpp
@@ -1,4 +1,4 @@
-// Copyright 1998-2014 Epic Games, Inc. All Rights Reserved.
+// Copyright 1998-2015 Epic Games, Inc. All Rights Reserved.
 
 #include "LogVisualizer.h"
 #include "Misc/CoreMisc.h"
diff --git a/Engine/Source/Runtime/AIModule/Classes/DetourCrowdAIController.h b/Engine/Source/Runtime/AIModule/Classes/DetourCrowdAIController.h
--- a/Engine/Source/Runtime/AIModule/Classes/DetourCrowdAIController.h
+++ b/Engine/Source/Runtime/AIModule/Classes/DetourCrowdAIController.h
@@ -1,4 +1,4 @@
-// Copyright 1998-2014 Epic Games, Inc. All Rights Reserved.
+// Copyright 1998-2015 Epic Games, Inc. All Rights Reserved.
 
 #pragma once
 
