We can have nest multiple levels of collapsible status information in LogVisualizer.

diff --git a/Engine/Source/Developer/LogVisualizer/Private/SVisualLoggerStatusView.cpp b/Engine/Source/Developer/LogVisualizer/Private/SVisualLoggerStatusView.cpp
--- a/Engine/Source/Developer/LogVisualizer/Private/SVisualLoggerStatusView.cpp
+++ b/Engine/Source/Developer/LogVisualizer/Private/SVisualLoggerStatusView.cpp
@@ -42,6 +42,25 @@ void SVisualLoggerStatusView::Construct(const FArguments& InArgs, const TSharedR
 		];
 }
 
+void GenerateChildren(TSharedPtr<FLogStatusItem> StatusItem, const FVisualLogStatusCategory LogCategory)
+{
+	for (int32 LineIndex = 0; LineIndex < LogCategory.Data.Num(); LineIndex++)
+	{
+		FString KeyDesc, ValueDesc;
+		const bool bHasValue = LogCategory.GetDesc(LineIndex, KeyDesc, ValueDesc);
+		if (bHasValue)
+		{
+			StatusItem->Children.Add(MakeShareable(new FLogStatusItem(KeyDesc, ValueDesc)));
+		}
+	}
+
+	for (const FVisualLogStatusCategory& Child : LogCategory.Children)
+	{
+		TSharedPtr<FLogStatusItem> ChildCategory = MakeShareable(new FLogStatusItem(Child.Category));
+		StatusItem->Children.Add(ChildCategory);
+		GenerateChildren(ChildCategory, Child);
+	}
+}
 void SVisualLoggerStatusView::OnItemSelectionChanged(const FVisualLogDevice::FVisualLogEntryItem& LogEntry)
 {
 	TArray<FString> ExpandedCategories;
@@ -61,7 +80,7 @@ void SVisualLoggerStatusView::OnItemSelectionChanged(const FVisualLogDevice::FVi
 
 	for (int32 CategoryIndex = 0; CategoryIndex < LogEntry.Entry.Status.Num(); CategoryIndex++)
 	{
-		if (LogEntry.Entry.Status[CategoryIndex].Data.Num() <= 0)
+		if (LogEntry.Entry.Status[CategoryIndex].Data.Num() <= 0 && LogEntry.Entry.Status[CategoryIndex].Children.Num() == 0)
 		{
 			continue;
 		}
@@ -77,6 +96,13 @@ void SVisualLoggerStatusView::OnItemSelectionChanged(const FVisualLogDevice::FVi
 			}
 		}
 
+		for (const FVisualLogStatusCategory& Child : LogEntry.Entry.Status[CategoryIndex].Children)
+		{
+			TSharedPtr<FLogStatusItem> ChildCategory = MakeShareable(new FLogStatusItem(Child.Category));
+			StatusItem->Children.Add(ChildCategory);
+			GenerateChildren(ChildCategory, Child);
+		}
+
 		StatusItems.Add(StatusItem);
 	}
 
diff --git a/Engine/Source/Runtime/Engine/Private/VisualLogger/VisualLoggerTypes.cpp b/Engine/Source/Runtime/Engine/Private/VisualLogger/VisualLoggerTypes.cpp
--- a/Engine/Source/Runtime/Engine/Private/VisualLogger/VisualLoggerTypes.cpp
+++ b/Engine/Source/Runtime/Engine/Private/VisualLogger/VisualLoggerTypes.cpp
@@ -312,6 +312,28 @@ FArchive& operator<<(FArchive& Ar, FVisualLogStatusCategory& Status)
 	Ar << Status.Category;
 	Ar << Status.Data;
 
+	const int32 VLogsVer = Ar.CustomVer(EVisualLoggerVersion::GUID);
+	if (VLogsVer >= EVisualLoggerVersion::StatusCategoryWithChildren)
+	{
+		int32 NumChildren = Status.Children.Num();
+		Ar << NumChildren;
+		if (Ar.IsLoading())
+		{
+			for (int32 Index = 0; Index < NumChildren; ++Index)
+			{
+				FVisualLogStatusCategory CurrentChild;
+				Ar << CurrentChild;
+				Status.Children.Add(CurrentChild);
+			}
+		}
+		else
+		{
+			for (auto& CurrentChild : Status.Children)
+			{
+				Ar << CurrentChild;
+			}
+		}
+	}
 	return Ar;
 }
 
diff --git a/Engine/Source/Runtime/Engine/Public/VisualLogger/VisualLoggerTypes.h b/Engine/Source/Runtime/Engine/Public/VisualLogger/VisualLoggerTypes.h
--- a/Engine/Source/Runtime/Engine/Public/VisualLogger/VisualLoggerTypes.h
+++ b/Engine/Source/Runtime/Engine/Public/VisualLogger/VisualLoggerTypes.h
@@ -28,6 +28,7 @@ namespace EVisualLoggerVersion
 		Initial = 0,
 		HistogramGraphsSerialization = 1,
 		AddedOwnerClassName = 2,
+		StatusCategoryWithChildren = 3,
 		// -----<new versions can be added before this line>-------------------------------------------------
 		// - this needs to be the last line (see note below)
 		VersionPlusOne,
@@ -101,9 +102,11 @@ struct ENGINE_API FVisualLogStatusCategory
 	TArray<FString> Data;
 	FString Category;
 	int32 UniqueId;
+	TArray<FVisualLogStatusCategory> Children;
 
 	void Add(const FString& Key, const FString& Value);
 	bool GetDesc(int32 Index, FString& Key, FString& Value) const;
+	void AddChild(const FVisualLogStatusCategory& Child);
 };
 
 struct ENGINE_API FVisualLogShapeElement
@@ -311,6 +314,12 @@ void FVisualLogStatusCategory::Add(const FString& Key, const FString& Value)
 	Data.Add(FString(Key).AppendChar(TEXT('|')) + Value);
 }
 
+FORCEINLINE
+void FVisualLogStatusCategory::AddChild(const FVisualLogStatusCategory& Child)
+{
+	Children.Add(Child);
+}
+
 FORCEINLINE
 FVisualLogShapeElement::FVisualLogShapeElement(EVisualLoggerShapeElement InType)
 : Verbosity(ELogVerbosity::All)
