Implemented option to stay always on the last frame in LogVisualizer. It can be enabled in LogVisuallizer settings.

diff --git a/Engine/Source/Developer/LogVisualizer/Private/STimeline.cpp b/Engine/Source/Developer/LogVisualizer/Private/STimeline.cpp
--- a/Engine/Source/Developer/LogVisualizer/Private/STimeline.cpp
+++ b/Engine/Source/Developer/LogVisualizer/Private/STimeline.cpp
@@ -176,6 +176,11 @@ void STimeline::AddEntry(const FVisualLogDevice::FVisualLogEntryItem& Entry)
 	ULogVisualizerSettings* Settings = ULogVisualizerSettings::StaticClass()->GetDefaultObject<ULogVisualizerSettings>();
 	UpdateVisibilityForEntry(Entries[Entries.Num() - 1], Settings->CurrentPresets.DataFilter, Settings->bSearchInsideLogs);
 	UpdateVisibility();
+
+	if (Settings->bStickToRecentData && IsSelected())
+	{
+		TimelineBar->SnapScrubPosition(Entry.Entry.TimeStamp);
+	}
 }
 
 void STimeline::Construct(const FArguments& InArgs, TSharedPtr<SVisualLoggerView> VisualLoggerView, TSharedPtr<FSequencerTimeSliderController> TimeSliderController, TSharedPtr<STimelinesContainer> InContainer, const FVisualLogDevice::FVisualLogEntryItem& Entry)
diff --git a/Engine/Source/Developer/LogVisualizer/Public/LogVisualizerSettings.h b/Engine/Source/Developer/LogVisualizer/Public/LogVisualizerSettings.h
--- a/Engine/Source/Developer/LogVisualizer/Public/LogVisualizerSettings.h
+++ b/Engine/Source/Developer/LogVisualizer/Public/LogVisualizerSettings.h
@@ -46,7 +46,7 @@ public:
 	int32 TrivialLogsThreshold;
 
 	/**Whether to show the recent data or not. Property disabled for now.*/
-	UPROPERTY(VisibleAnywhere, config, Category = "VisualLogger")
+	UPROPERTY(EditAnywhere, config, Category = "VisualLogger")
 	bool bStickToRecentData;
 
 	/**Whether to show histogram labels inside graph or outside. Property disabled for now.*/
