Feature request for LogVisualizer:
- data reset for new PIE sessions (or for each new data loaded from file). This feature is disabled by default.
- fixes for  "stick to new data" feature, to work better with zoomed data, etc.
- tweaks and fixes for movement on timeline with keyboard
- fixed time range issues on timellines

diff --git a/Engine/Source/Developer/LogVisualizer/Private/LogVisualizerSettings.cpp b/Engine/Source/Developer/LogVisualizer/Private/LogVisualizerSettings.cpp
--- a/Engine/Source/Developer/LogVisualizer/Private/LogVisualizerSettings.cpp
+++ b/Engine/Source/Developer/LogVisualizer/Private/LogVisualizerSettings.cpp
@@ -16,6 +16,7 @@ ULogVisualizerSettings::ULogVisualizerSettings(const FObjectInitializer& ObjectI
 	DefaultCameraDistance = 150;
 	bSearchInsideLogs = true;
 	GraphsBackgroundColor = FColor(0, 0, 0, 70);
+	bResetDataWithNewSession = false;
 }
 
 #if WITH_EDITOR
diff --git a/Engine/Source/Developer/LogVisualizer/Private/LogVisualizerStyle.cpp b/Engine/Source/Developer/LogVisualizer/Private/LogVisualizerStyle.cpp
--- a/Engine/Source/Developer/LogVisualizer/Private/LogVisualizerStyle.cpp
+++ b/Engine/Source/Developer/LogVisualizer/Private/LogVisualizerStyle.cpp
@@ -81,7 +81,7 @@ TSharedRef< FSlateStyleSet > FLogVisualizerStyle::Create()
 		Style.Set("Toolbar.Camera", new IMAGE_BRUSH("Icons/LV_Camera", Icon40x40));
 		Style.Set("Toolbar.Save", new IMAGE_BRUSH("Icons/LV_Save", Icon40x40));
 		Style.Set("Toolbar.Load", new IMAGE_BRUSH("Icons/LV_Load", Icon40x40));
-		Style.Set("Toolbar.Remove", new IMAGE_BRUSH("Icons/LV_Remove", Icon40x40));
+		Style.Set("Toolbar.Remove", new IMAGE_BRUSH("Icons/Edit/icon_Edit_Delete_40x", Icon40x40));
 		Style.Set("Toolbar.Graphs", new IMAGE_BRUSH("Icons/icon_MatEd_Stats_40x", Icon40x40));
 	}
 
diff --git a/Engine/Source/Developer/LogVisualizer/Private/STimeline.cpp b/Engine/Source/Developer/LogVisualizer/Private/STimeline.cpp
--- a/Engine/Source/Developer/LogVisualizer/Private/STimeline.cpp
+++ b/Engine/Source/Developer/LogVisualizer/Private/STimeline.cpp
@@ -191,9 +191,6 @@ void STimeline::Construct(const FArguments& InArgs, TSharedPtr<SVisualLoggerView
 	Owner = InContainer;
 	Name = Entry.OwnerName;
 	OwnerClassName = Entry.OwnerClassName;
-
-	Entries.Add(Entry);
-	OnFiltersChanged();
 	
 	ULogVisualizerSettings::StaticClass()->GetDefaultObject<ULogVisualizerSettings>()->OnSettingChanged().AddRaw(this, &STimeline::HandleLogVisualizerSettingChanged);
 
@@ -258,6 +255,8 @@ void STimeline::Construct(const FArguments& InArgs, TSharedPtr<SVisualLoggerView
 				]
 			]
 		];
+
+	AddEntry(Entry);
 }
 
 const FSlateBrush* STimeline::GetBorder() const
diff --git a/Engine/Source/Developer/LogVisualizer/Private/STimelineBar.cpp b/Engine/Source/Developer/LogVisualizer/Private/STimelineBar.cpp
--- a/Engine/Source/Developer/LogVisualizer/Private/STimelineBar.cpp
+++ b/Engine/Source/Developer/LogVisualizer/Private/STimelineBar.cpp
@@ -179,14 +179,6 @@ FReply STimelineBar::OnKeyDown(const FGeometry& MyGeometry, const FKeyEvent& InK
 		{
 			float NewTimeStamp = Entries[NewItemIndex].Entry.TimeStamp;
 			SnapScrubPosition(NewItemIndex);
-			if (NewTimeStamp < LocalViewRange.GetLowerBoundValue())
-			{
-				TimeSliderController->SetTimeRange(NewTimeStamp, NewTimeStamp + RangeSize);
-			}
-			else if (NewTimeStamp > LocalViewRange.GetUpperBoundValue())
-			{
-				TimeSliderController->SetTimeRange(NewTimeStamp - RangeSize, NewTimeStamp);
-			}
 		}
 	}
 
diff --git a/Engine/Source/Developer/LogVisualizer/Private/STimelinesContainer.cpp b/Engine/Source/Developer/LogVisualizer/Private/STimelinesContainer.cpp
--- a/Engine/Source/Developer/LogVisualizer/Private/STimelinesContainer.cpp
+++ b/Engine/Source/Developer/LogVisualizer/Private/STimelinesContainer.cpp
@@ -253,6 +253,20 @@ FReply STimelinesContainer::OnKeyDown(const FGeometry& MyGeometry, const FKeyEve
 	return FReply::Unhandled();
 }
 
+void STimelinesContainer::ResetData()
+{
+	for (auto CurrentItem : TimelineItems)
+	{
+		ContainingBorder->RemoveSlot(CurrentItem.ToSharedRef());
+	}
+	TimelineItems.Reset();
+
+	CachedMinTime = FLT_MAX;
+	CachedMaxTime = 0;
+	TimeSliderController->SetClampRange(0, 5);
+	TimeSliderController->SetTimeRange(0, 5);
+}
+
 void STimelinesContainer::Construct(const FArguments& InArgs, TSharedRef<class SVisualLoggerView> InVisualLoggerView, TSharedRef<FSequencerTimeSliderController> InTimeSliderController)
 {
 	TimeSliderController = InTimeSliderController;
@@ -269,6 +283,10 @@ void STimelinesContainer::Construct(const FArguments& InArgs, TSharedRef<class S
 				SAssignNew(ContainingBorder, SVerticalBox)
 			]
 		];
+
+	CachedMinTime = FLT_MAX;
+	CachedMaxTime = 0;
+
 }
 
 void STimelinesContainer::OnSearchChanged(const FText& Filter)
@@ -297,13 +315,16 @@ void STimelinesContainer::OnNewLogEntry(const FVisualLogDevice::FVisualLogEntryI
 		MakeTimeline(VisualLoggerView, TimeSliderController, Entry);
 	}
 
+	CachedMinTime = CachedMinTime < Entry.Entry.TimeStamp ? CachedMinTime : Entry.Entry.TimeStamp;
+	CachedMaxTime = CachedMaxTime > Entry.Entry.TimeStamp ? CachedMaxTime : Entry.Entry.TimeStamp;
+
 	TRange<float> LocalViewRange = TimeSliderController->GetTimeSliderArgs().ViewRange.Get();
 	const float CurrentMin = TimeSliderController->GetTimeSliderArgs().ClampMin.Get().GetValue();
 	const float CurrentMax = TimeSliderController->GetTimeSliderArgs().ClampMax.Get().GetValue();
 	float ZoomLevel = LocalViewRange.Size<float>() / (CurrentMax - CurrentMin);
 
-	TimeSliderController->GetTimeSliderArgs().ClampMin = FMath::Min(0.0f, CurrentMin);
-	TimeSliderController->GetTimeSliderArgs().ClampMax = FMath::Max(Entry.Entry.TimeStamp, CurrentMax);
+	TimeSliderController->GetTimeSliderArgs().ClampMin = CachedMinTime;
+	TimeSliderController->GetTimeSliderArgs().ClampMax = CachedMaxTime + 0.1;
 	if ( FMath::Abs(ZoomLevel - 1) <= SMALL_NUMBER)
 	{
 		TimeSliderController->GetTimeSliderArgs().ViewRange = TRange<float>(TimeSliderController->GetTimeSliderArgs().ClampMin.Get().GetValue(), TimeSliderController->GetTimeSliderArgs().ClampMax.Get().GetValue());
diff --git a/Engine/Source/Developer/LogVisualizer/Private/STimelinesContainer.h b/Engine/Source/Developer/LogVisualizer/Private/STimelinesContainer.h
--- a/Engine/Source/Developer/LogVisualizer/Private/STimelinesContainer.h
+++ b/Engine/Source/Developer/LogVisualizer/Private/STimelinesContainer.h
@@ -32,6 +32,8 @@ public:
 	void OnSearchChanged(const FText& Filter);
 	void OnFiltersSearchChanged(const FText& Filter);
 
+	void ResetData();
+
 	void GenerateReport();
 
 	/**
@@ -65,4 +67,5 @@ protected:
 	TArray< TSharedPtr<class STimeline> > SelectedNodes;
 	TSharedPtr<SVerticalBox> ContainingBorder;
 	TWeakPtr<IVisualLoggerInterface> VisualLoggerInterface;
+	float CachedMinTime, CachedMaxTime;
 };
diff --git a/Engine/Source/Developer/LogVisualizer/Private/SVisualLogger.cpp b/Engine/Source/Developer/LogVisualizer/Private/SVisualLogger.cpp
--- a/Engine/Source/Developer/LogVisualizer/Private/SVisualLogger.cpp
+++ b/Engine/Source/Developer/LogVisualizer/Private/SVisualLogger.cpp
@@ -138,6 +138,13 @@ void SVisualLogger::FVisualLoggerDevice::Serialize(const class UObject* LogOwner
 		return;
 	}
 
+	UWorld* World = LogOwner != nullptr ? GEngine->GetWorldFromContextObject(LogOwner) : NULL;
+	if (LastWorld.Get() != World)
+	{
+		Owner->OnNewWorld(LastWorld.Get());
+		LastWorld = World;
+	}
+
 	Owner->OnNewLogEntry(FVisualLogDevice::FVisualLogEntryItem(OwnerName, OwnerClassName, LogEntry));
 }
 
@@ -247,6 +254,11 @@ void SVisualLogger::Construct(const FArguments& InArgs, const TSharedRef<SDockTa
 		FCanExecuteAction(),
 		FIsActionChecked::CreateLambda([]()->bool{return ULogVisualizerSessionSettings::StaticClass()->GetDefaultObject<ULogVisualizerSessionSettings>()->bEnableGraphsVisualization; }),
 		FIsActionButtonVisible());
+	ActionList.MapAction(Commands.ResetData, 
+		FExecuteAction::CreateRaw(this, &SVisualLogger::ResetData),
+		FCanExecuteAction::CreateRaw(this, &SVisualLogger::HandleSaveCommandCanExecute),
+		FIsActionChecked(), 
+		FIsActionButtonVisible::CreateRaw(this, &SVisualLogger::HandleSaveCommandCanExecute));
 
 
 	// Tab Spawners
@@ -643,6 +655,7 @@ void SVisualLogger::HandleLoadCommandExecute()
 
 	if (bOpened && OpenFilenames.Num() > 0)
 	{
+		OnNewWorld(nullptr);
 		for (int FilenameIndex = 0; FilenameIndex < OpenFilenames.Num(); ++FilenameIndex)
 		{
 			FString CurrentFileName = OpenFilenames[FilenameIndex];
@@ -732,6 +745,30 @@ void SVisualLogger::HandleSaveCommandExecute()
 	}
 }
 
+void SVisualLogger::ResetData()
+{
+	if (ULogVisualizerSettings::StaticClass()->GetDefaultObject<ULogVisualizerSettings>()->bResetDataWithNewSession)
+	{
+		return;
+	}
+
+	if (MainView.IsValid())
+	{
+		MainView->ResetData();
+	}
+
+	if (VisualLoggerFilters.IsValid())
+	{
+		VisualLoggerFilters->ResetData();
+	}
+}
+
+void SVisualLogger::OnNewWorld(UWorld* NewWorld)
+{
+	InternalDevice->SerLastWorld(NewWorld);
+	ResetData();
+}
+
 void SVisualLogger::OnNewLogEntry(const FVisualLogDevice::FVisualLogEntryItem& Entry)
 {
 	CollectNewCategories(Entry);
diff --git a/Engine/Source/Developer/LogVisualizer/Private/SVisualLogger.h b/Engine/Source/Developer/LogVisualizer/Private/SVisualLogger.h
--- a/Engine/Source/Developer/LogVisualizer/Private/SVisualLogger.h
+++ b/Engine/Source/Developer/LogVisualizer/Private/SVisualLogger.h
@@ -37,13 +37,17 @@ public:
 public:
 	class FVisualLoggerDevice : public FVisualLogDevice
 	{
+		friend class SVisualLogger;
 	public:
 		FVisualLoggerDevice(SVisualLogger* Owner);
 		virtual ~FVisualLoggerDevice();
 		virtual void Serialize(const class UObject* LogOwner, FName OwnerName, FName OwnerClassName, const FVisualLogEntry& LogEntry) override;
 
 	protected:
+		void SerLastWorld(class UWorld* InWorld) { LastWorld = InWorld; }
+
 		SVisualLogger* Owner;
+		TWeakObjectPtr<class UWorld>	LastWorld;
 	};
 	friend class FVisualLogDevice;
 
@@ -93,6 +97,8 @@ public:
 	void HandleTabManagerPersistLayout(const TSharedRef<FTabManager::FLayout>& LayoutToSave);
 
 	void SetFiltersPreset(const struct FFiltersPreset& Preset);
+	void OnNewWorld(UWorld* NewWorld);
+	void ResetData();
 
 protected:
 	// Holds the list of UI commands.
diff --git a/Engine/Source/Developer/LogVisualizer/Private/SVisualLoggerFilters.cpp b/Engine/Source/Developer/LogVisualizer/Private/SVisualLoggerFilters.cpp
--- a/Engine/Source/Developer/LogVisualizer/Private/SVisualLoggerFilters.cpp
+++ b/Engine/Source/Developer/LogVisualizer/Private/SVisualLoggerFilters.cpp
@@ -388,6 +388,16 @@ void SVisualLoggerFilters::AddFilter(const FString& InFilterName)
 	GraphsFilterCombo->SetVisibility(GraphFilters.Num() > 0 ? EVisibility::Visible : EVisibility::Collapsed);
 }
 
+void SVisualLoggerFilters::ResetData()
+{
+	for (auto& CurrentFilter : Filters)
+	{
+		FilterBox->RemoveSlot(CurrentFilter);
+	}
+	Filters.Reset();
+}
+
+
 void SVisualLoggerFilters::OnFiltersChanged()
 {
 	TArray<FString> EnabledFilters;
diff --git a/Engine/Source/Developer/LogVisualizer/Private/SVisualLoggerFilters.h b/Engine/Source/Developer/LogVisualizer/Private/SVisualLoggerFilters.h
--- a/Engine/Source/Developer/LogVisualizer/Private/SVisualLoggerFilters.h
+++ b/Engine/Source/Developer/LogVisualizer/Private/SVisualLoggerFilters.h
@@ -44,6 +44,7 @@ public:
 	void GraphFilterCategoryClicked(FName MenuCategory);
 	bool IsGraphFilterCategoryInUse(FName MenuCategory) const;
 	void OnFiltersChanged();
+	void ResetData();
 
 	TSharedRef<SWidget> MakeGraphsFilterMenu();
 	void CreateFiltersMenuCategoryForGraph(FMenuBuilder& MenuBuilder, FName MenuCategory) const;
diff --git a/Engine/Source/Developer/LogVisualizer/Private/SVisualLoggerStatusView.cpp b/Engine/Source/Developer/LogVisualizer/Private/SVisualLoggerStatusView.cpp
--- a/Engine/Source/Developer/LogVisualizer/Private/SVisualLoggerStatusView.cpp
+++ b/Engine/Source/Developer/LogVisualizer/Private/SVisualLoggerStatusView.cpp
@@ -63,15 +63,8 @@ void GenerateChildren(TSharedPtr<FLogStatusItem> StatusItem, const FVisualLogSta
 }
 void SVisualLoggerStatusView::OnItemSelectionChanged(const FVisualLogDevice::FVisualLogEntryItem& LogEntry)
 {
-	TArray<FString> ExpandedCategories;
-	for (int32 ItemIndex = 0; ItemIndex < StatusItems.Num(); ItemIndex++)
-	{
-		const bool bIsExpanded = StatusItemsView->IsItemExpanded(StatusItems[ItemIndex]);
-		if (bIsExpanded)
-		{
-			ExpandedCategories.Add(StatusItems[ItemIndex]->ItemText);
-		}
-	}
+	//TSet<TSharedPtr<FLogStatusItem> > ExpandedItems;
+	//StatusItemsView->GetExpandedItems(ExpandedItems);
 
 	StatusItems.Empty();
 
@@ -108,17 +101,10 @@ void SVisualLoggerStatusView::OnItemSelectionChanged(const FVisualLogDevice::FVi
 
 	StatusItemsView->RequestTreeRefresh();
 
-	for (int32 ItemIndex = 0; ItemIndex < StatusItems.Num(); ItemIndex++)
-	{
-		for (const FString& Category : ExpandedCategories)
-		{
-			if (StatusItems[ItemIndex]->ItemText == Category)
-			{
-				StatusItemsView->SetItemExpansion(StatusItems[ItemIndex], true);
-				break;
-			}
-		}
-	}
+	//for (auto& CurrentItem : ExpandedItems)
+	//{
+	//	StatusItemsView->SetItemExpansion(CurrentItem, true);
+	//}
 }
 
 void SVisualLoggerStatusView::OnLogStatusGetChildren(TSharedPtr<FLogStatusItem> InItem, TArray< TSharedPtr<FLogStatusItem> >& OutItems)
diff --git a/Engine/Source/Developer/LogVisualizer/Private/SVisualLoggerToolbar.cpp b/Engine/Source/Developer/LogVisualizer/Private/SVisualLoggerToolbar.cpp
--- a/Engine/Source/Developer/LogVisualizer/Private/SVisualLoggerToolbar.cpp
+++ b/Engine/Source/Developer/LogVisualizer/Private/SVisualLoggerToolbar.cpp
@@ -33,6 +33,7 @@ TSharedRef<SWidget> SVisualLoggerToolbar::MakeToolbar(const TSharedRef<FUIComman
 		ToolBarBuilder.AddToolBarButton(FVisualLoggerCommands::Get().Save, NAME_None, LOCTEXT("Save", "Save"), TAttribute<FText>(), FSlateIcon(FLogVisualizerStyle::Get().GetStyleSetName(), TEXT("Toolbar.Save")));
 		ToolBarBuilder.AddSeparator();
 		ToolBarBuilder.AddToolBarButton(FVisualLoggerCommands::Get().ToggleGraphs, NAME_None, LOCTEXT("ToggleGraphs", "Graphs"), TAttribute<FText>(), FSlateIcon(FLogVisualizerStyle::Get().GetStyleSetName(), TEXT("Toolbar.Graphs")));
+		ToolBarBuilder.AddToolBarButton(FVisualLoggerCommands::Get().ResetData, NAME_None, LOCTEXT("ResetData", "Clear"), TAttribute<FText>(), FSlateIcon(FLogVisualizerStyle::Get().GetStyleSetName(), TEXT("Toolbar.Remove")));
 	}
 
 	return ToolBarBuilder.MakeWidget();
diff --git a/Engine/Source/Developer/LogVisualizer/Private/SVisualLoggerView.cpp b/Engine/Source/Developer/LogVisualizer/Private/SVisualLoggerView.cpp
--- a/Engine/Source/Developer/LogVisualizer/Private/SVisualLoggerView.cpp
+++ b/Engine/Source/Developer/LogVisualizer/Private/SVisualLoggerView.cpp
@@ -336,6 +336,11 @@ TSharedRef<SWidget> SVisualLoggerView::MakeSectionOverlay(TSharedRef<FSequencerT
 		];
 }
 
+void SVisualLoggerView::ResetData()
+{
+	TimelinesContainer->ResetData();
+}
+
 void SVisualLoggerView::OnNewLogEntry(const FVisualLogDevice::FVisualLogEntryItem& Entry)
 {
 	TimelinesContainer->OnNewLogEntry(Entry);
diff --git a/Engine/Source/Developer/LogVisualizer/Private/SVisualLoggerView.h b/Engine/Source/Developer/LogVisualizer/Private/SVisualLoggerView.h
--- a/Engine/Source/Developer/LogVisualizer/Private/SVisualLoggerView.h
+++ b/Engine/Source/Developer/LogVisualizer/Private/SVisualLoggerView.h
@@ -39,6 +39,7 @@ public:
 	void GetTimelines(TArray<TSharedPtr<class STimeline> >&, bool bOnlySelectedOnes = false);
 	virtual FCursorReply OnCursorQuery(const FGeometry& MyGeometry, const FPointerEvent& CursorEvent) const override;
 	TSharedRef<SWidget> MakeClassesFilterMenu();
+	void ResetData();
 
 protected:
 	TWeakPtr<struct IVisualLoggerInterface> VisualLoggerInterface;
diff --git a/Engine/Source/Developer/LogVisualizer/Private/TimeSliderController.cpp b/Engine/Source/Developer/LogVisualizer/Private/TimeSliderController.cpp
--- a/Engine/Source/Developer/LogVisualizer/Private/TimeSliderController.cpp
+++ b/Engine/Source/Developer/LogVisualizer/Private/TimeSliderController.cpp
@@ -420,7 +420,7 @@ FReply FSequencerTimeSliderController::OnMouseMove( TSharedRef<SWidget> WidgetOw
 				TimeSliderArgs.OnViewRangeChanged.ExecuteIfBound(TRange<float>(NewViewOutputMin, NewViewOutputMax));
 				if (Scrollbar.IsValid())
 				{
-					float InOffsetFraction = NewViewOutputMin / (LocalClampMax.GetValue() - LocalClampMin.GetValue());
+					float InOffsetFraction = (NewViewOutputMin - LocalClampMin.GetValue()) / (LocalClampMax.GetValue() - LocalClampMin.GetValue());
 					float InThumbSizeFraction = (NewViewOutputMax - NewViewOutputMin) / (LocalClampMax.GetValue() - LocalClampMin.GetValue());
 					Scrollbar->SetState(InOffsetFraction, InThumbSizeFraction);
 				}
@@ -466,6 +466,17 @@ void FSequencerTimeSliderController::CommitScrubPosition( float NewValue, bool b
 		TimeSliderArgs.ScrubPosition.Set( NewValue );
 	}
 
+	TRange<float> LocalViewRange = TimeSliderArgs.ViewRange.Get();
+	const float RangeSize = LocalViewRange.Size<float>();
+	if (NewValue < LocalViewRange.GetLowerBoundValue())
+	{
+		SetTimeRange(NewValue, NewValue + RangeSize);
+	}
+	else if (NewValue > LocalViewRange.GetUpperBoundValue())
+	{
+		SetTimeRange(NewValue - RangeSize, NewValue);
+	}
+
 	TimeSliderArgs.OnScrubPositionChanged.ExecuteIfBound( NewValue, bIsScrubbing );
 }
 
@@ -485,14 +496,19 @@ void FSequencerTimeSliderController::HorizontalScrollBar_OnUserScrolled(float Sc
 		TOptional<float> LocalClampMin = TimeSliderArgs.ClampMin.Get();
 		TOptional<float> LocalClampMax = TimeSliderArgs.ClampMax.Get();
 
-		float NewViewOutputMin = ScrollOffset * (LocalClampMax.GetValue() - LocalClampMin.GetValue());
+		float InThumbSizeFraction = (LocalViewRangeMax - LocalViewRangeMin) / (LocalClampMax.GetValue() - LocalClampMin.GetValue());
+
+		float NewViewOutputMin = LocalClampMin.GetValue() + ScrollOffset * (LocalClampMax.GetValue() - LocalClampMin.GetValue());
 		// The  output is not bound to a delegate so we'll manage the value ourselves
-		float NewViewOutputMax = NewViewOutputMin + (LocalViewRangeMax - LocalViewRangeMin);
-		TimeSliderArgs.ViewRange.Set(TRange<float>(NewViewOutputMin, NewViewOutputMax));
+		float NewViewOutputMax = FMath::Min<float>(NewViewOutputMin + (LocalViewRangeMax - LocalViewRangeMin), LocalClampMax.GetValue());
+		NewViewOutputMin = NewViewOutputMax - (LocalViewRangeMax - LocalViewRangeMin);
 
-		float InOffsetFraction = NewViewOutputMin / (LocalClampMax.GetValue() - LocalClampMin.GetValue());
-		float InThumbSizeFraction = (NewViewOutputMax - NewViewOutputMin) / (LocalClampMax.GetValue() - LocalClampMin.GetValue());
-		Scrollbar->SetState(InOffsetFraction, InThumbSizeFraction);
+		float InOffsetFraction = (NewViewOutputMin - LocalClampMin.GetValue()) / (LocalClampMax.GetValue() - LocalClampMin.GetValue());
+		//if (InOffsetFraction + InThumbSizeFraction <= 1)
+		{
+			TimeSliderArgs.ViewRange.Set(TRange<float>(NewViewOutputMin, NewViewOutputMax));
+			Scrollbar->SetState(InOffsetFraction, InThumbSizeFraction);
+		}
 	}
 }
 
@@ -503,11 +519,29 @@ void FSequencerTimeSliderController::SetTimeRange(float NewViewOutputMin, float
 	TOptional<float> LocalClampMin = TimeSliderArgs.ClampMin.Get();
 	TOptional<float> LocalClampMax = TimeSliderArgs.ClampMax.Get();
 
-	float InOffsetFraction = NewViewOutputMin / (LocalClampMax.GetValue() - LocalClampMin.GetValue());
-	float InThumbSizeFraction = (NewViewOutputMax - NewViewOutputMin) / (LocalClampMax.GetValue() - LocalClampMin.GetValue());
+	const float InOffsetFraction = (NewViewOutputMin - LocalClampMin.GetValue()) / (LocalClampMax.GetValue() - LocalClampMin.GetValue());
+	const float InThumbSizeFraction = (NewViewOutputMax - NewViewOutputMin) / (LocalClampMax.GetValue() - LocalClampMin.GetValue());
 	Scrollbar->SetState(InOffsetFraction, InThumbSizeFraction);
 }
 
+void FSequencerTimeSliderController::SetClampRange(float MinValue, float MaxValue)
+{
+	TRange<float> LocalViewRange = TimeSliderArgs.ViewRange.Get();
+	float LocalClampMin = TimeSliderArgs.ClampMin.Get().GetValue();
+	float LocalClampMax = TimeSliderArgs.ClampMax.Get().GetValue();
+	const float CurrentDistance = LocalClampMax - LocalClampMin;
+	const float ZoomDelta = (LocalViewRange.GetUpperBoundValue() - LocalViewRange.GetLowerBoundValue()) / CurrentDistance;
+
+	MaxValue = MinValue + (MaxValue - MinValue < 2 ? CurrentDistance : MaxValue - MinValue);
+
+	TimeSliderArgs.ClampMin.Set(MinValue);
+	TimeSliderArgs.ClampMax.Set(MaxValue);
+
+	const float LocalViewRangeMin = FMath::Clamp(LocalViewRange.GetLowerBoundValue(), MinValue, MaxValue);
+	const float LocalViewRangeMax = FMath::Clamp(LocalViewRange.GetUpperBoundValue(), MinValue, MaxValue);
+	SetTimeRange(ZoomDelta >= 1 ? MinValue : LocalViewRangeMin, ZoomDelta >= 1 ? MaxValue : LocalViewRangeMax);
+}
+
 FReply FSequencerTimeSliderController::OnMouseWheel( TSharedRef<SWidget> WidgetOwner, const FGeometry& MyGeometry, const FPointerEvent& MouseEvent )
 {
 	FReply ReturnValue = FReply::Unhandled();;
@@ -552,15 +586,7 @@ FReply FSequencerTimeSliderController::OnMouseWheel( TSharedRef<SWidget> WidgetO
 				TimeSliderArgs.OnViewRangeChanged.ExecuteIfBound(TRange<float>(NewViewOutputMin, NewViewOutputMax));
 				if (Scrollbar.IsValid())
 				{
-					/**
-					* Set the offset and size of the track's thumb.
-					* Note that the maximum offset is 1.0-ThumbSizeFraction.
-					* If the user can view 1/3 of the items in a single page, the maximum offset will be ~0.667f
-					*
-					* @param InOffsetFraction     Offset of the thumbnail from the top as a fraction of the total available scroll space.
-					* @param InThumbSizeFraction  Size of thumbnail as a fraction of the total available scroll space.
-					*/
-					float InOffsetFraction = NewViewOutputMin / (LocalClampMax.GetValue() - LocalClampMin.GetValue());
+					float InOffsetFraction = (NewViewOutputMin - LocalClampMin.GetValue()) / (LocalClampMax.GetValue() - LocalClampMin.GetValue());
 					float InThumbSizeFraction = (NewViewOutputMax - NewViewOutputMin) / (LocalClampMax.GetValue() - LocalClampMin.GetValue());
 					Scrollbar->SetState(InOffsetFraction, InThumbSizeFraction);
 				}
diff --git a/Engine/Source/Developer/LogVisualizer/Private/TimeSliderController.h b/Engine/Source/Developer/LogVisualizer/Private/TimeSliderController.h
--- a/Engine/Source/Developer/LogVisualizer/Private/TimeSliderController.h
+++ b/Engine/Source/Developer/LogVisualizer/Private/TimeSliderController.h
@@ -8,7 +8,9 @@ struct FVisualLoggerTimeSliderArgs : FTimeSliderArgs
 {
 	FVisualLoggerTimeSliderArgs()
 		: CursorSize(0.1f)
-	{}
+	{
+		ViewRange = TRange<float>(0.0f, 5.0f);
+	}
 
 	/** Cursor range for data like histogram graphs, etc. */
 	TAttribute< float > CursorSize;
@@ -45,6 +47,7 @@ public:
 
 	void SetExternalScrollbar(TSharedRef<SScrollBar> Scrollbar);
 	void SetTimeRange(float MinValue, float MaxValue);
+	void SetClampRange(float MinValue, float MaxValue);
 	bool IsPanning() { return bPanning; }
 
 private:
diff --git a/Engine/Source/Developer/LogVisualizer/Private/VisualLoggerCommands.h b/Engine/Source/Developer/LogVisualizer/Private/VisualLoggerCommands.h
--- a/Engine/Source/Developer/LogVisualizer/Private/VisualLoggerCommands.h
+++ b/Engine/Source/Developer/LogVisualizer/Private/VisualLoggerCommands.h
@@ -28,6 +28,7 @@ public:
 		UI_COMMAND(Load, "Load", "Load external vlogs", EUserInterfaceActionType::Button, FInputGesture());
 		UI_COMMAND(Save, "Save", "Save selected data to vlog file", EUserInterfaceActionType::Button, FInputGesture());
 		UI_COMMAND(ToggleGraphs, "ToggleGraphs", "Toggle graphs visualization on/off", EUserInterfaceActionType::ToggleButton, FInputGesture());
+		UI_COMMAND(ResetData, "Clear", "Clear all data", EUserInterfaceActionType::Button, FInputGesture());
 	}
 
 public:
@@ -40,6 +41,7 @@ public:
 	TSharedPtr<FUICommandInfo> Load;
 	TSharedPtr<FUICommandInfo> Save;
 	TSharedPtr<FUICommandInfo> ToggleGraphs;
+	TSharedPtr<FUICommandInfo> ResetData;
 };
 
 
diff --git a/Engine/Source/Developer/LogVisualizer/Public/LogVisualizerSettings.h b/Engine/Source/Developer/LogVisualizer/Public/LogVisualizerSettings.h
--- a/Engine/Source/Developer/LogVisualizer/Public/LogVisualizerSettings.h
+++ b/Engine/Source/Developer/LogVisualizer/Public/LogVisualizerSettings.h
@@ -49,6 +49,10 @@ public:
 	UPROPERTY(EditAnywhere, config, Category = "VisualLogger")
 	bool bStickToRecentData;
 
+	/**Whether to reset current data or not for each new session.*/
+	UPROPERTY(EditAnywhere, config, Category = "VisualLogger")
+	bool bResetDataWithNewSession;
+
 	/**Whether to show histogram labels inside graph or outside. Property disabled for now.*/
 	UPROPERTY(VisibleAnywhere, config, Category = "VisualLogger")
 	bool bShowHistogramLabelsOutside;
