Fix visual logger command conflict with level save commands causing tooltips to appear incorrectly.

diff --git a/Engine/Source/Developer/LogVisualizer/Private/SVisualLogger.cpp b/Engine/Source/Developer/LogVisualizer/Private/SVisualLogger.cpp
--- a/Engine/Source/Developer/LogVisualizer/Private/SVisualLogger.cpp
+++ b/Engine/Source/Developer/LogVisualizer/Private/SVisualLogger.cpp
@@ -242,8 +242,8 @@ void SVisualLogger::Construct(const FArguments& InArgs, const TSharedRef<SDockTa
 	ActionList.MapAction(Commands.StopRecording, FExecuteAction::CreateRaw(this, &SVisualLogger::HandleStopRecordingCommandExecute), FCanExecuteAction::CreateRaw(this, &SVisualLogger::HandleStopRecordingCommandCanExecute), FIsActionChecked(), FIsActionButtonVisible::CreateRaw(this, &SVisualLogger::HandleStopRecordingCommandIsVisible));
 	ActionList.MapAction(Commands.Pause, FExecuteAction::CreateRaw(this, &SVisualLogger::HandlePauseCommandExecute), FCanExecuteAction::CreateRaw(this, &SVisualLogger::HandlePauseCommandCanExecute), FIsActionChecked(), FIsActionButtonVisible::CreateRaw(this, &SVisualLogger::HandlePauseCommandIsVisible));
 	ActionList.MapAction(Commands.Resume, FExecuteAction::CreateRaw(this, &SVisualLogger::HandleResumeCommandExecute), FCanExecuteAction::CreateRaw(this, &SVisualLogger::HandleResumeCommandCanExecute), FIsActionChecked(), FIsActionButtonVisible::CreateRaw(this, &SVisualLogger::HandleResumeCommandIsVisible));
-	ActionList.MapAction(Commands.Load, FExecuteAction::CreateRaw(this, &SVisualLogger::HandleLoadCommandExecute), FCanExecuteAction::CreateRaw(this, &SVisualLogger::HandleLoadCommandCanExecute), FIsActionChecked(), FIsActionButtonVisible::CreateRaw(this, &SVisualLogger::HandleLoadCommandCanExecute));
-	ActionList.MapAction(Commands.Save, FExecuteAction::CreateRaw(this, &SVisualLogger::HandleSaveCommandExecute), FCanExecuteAction::CreateRaw(this, &SVisualLogger::HandleSaveCommandCanExecute), FIsActionChecked(), FIsActionButtonVisible::CreateRaw(this, &SVisualLogger::HandleSaveCommandCanExecute));
+	ActionList.MapAction(Commands.LoadFromVLog, FExecuteAction::CreateRaw(this, &SVisualLogger::HandleLoadCommandExecute), FCanExecuteAction::CreateRaw(this, &SVisualLogger::HandleLoadCommandCanExecute), FIsActionChecked(), FIsActionButtonVisible::CreateRaw(this, &SVisualLogger::HandleLoadCommandCanExecute));
+	ActionList.MapAction(Commands.SaveToVLog, FExecuteAction::CreateRaw(this, &SVisualLogger::HandleSaveCommandExecute), FCanExecuteAction::CreateRaw(this, &SVisualLogger::HandleSaveCommandCanExecute), FIsActionChecked(), FIsActionButtonVisible::CreateRaw(this, &SVisualLogger::HandleSaveCommandCanExecute));
 	ActionList.MapAction(Commands.FreeCamera, 
 		FExecuteAction::CreateRaw(this, &SVisualLogger::HandleCameraCommandExecute), 
 		FCanExecuteAction::CreateRaw(this, &SVisualLogger::HandleCameraCommandCanExecute), 
@@ -412,8 +412,8 @@ void SVisualLogger::FillFileMenu(FMenuBuilder& MenuBuilder, const TSharedPtr<FTa
 {
 	MenuBuilder.BeginSection("LogFile", LOCTEXT("FileMenu", "Log File"));
 	{
-		MenuBuilder.AddMenuEntry(FVisualLoggerCommands::Get().Load);
-		MenuBuilder.AddMenuEntry(FVisualLoggerCommands::Get().Save);
+		MenuBuilder.AddMenuEntry(FVisualLoggerCommands::Get().LoadFromVLog);
+		MenuBuilder.AddMenuEntry(FVisualLoggerCommands::Get().SaveToVLog);
 	}
 	MenuBuilder.EndSection();
 	MenuBuilder.BeginSection("LogFilters", LOCTEXT("FIlterMenu", "Log Filters"));
diff --git a/Engine/Source/Developer/LogVisualizer/Private/SVisualLoggerToolbar.cpp b/Engine/Source/Developer/LogVisualizer/Private/SVisualLoggerToolbar.cpp
--- a/Engine/Source/Developer/LogVisualizer/Private/SVisualLoggerToolbar.cpp
+++ b/Engine/Source/Developer/LogVisualizer/Private/SVisualLoggerToolbar.cpp
@@ -29,8 +29,8 @@ TSharedRef<SWidget> SVisualLoggerToolbar::MakeToolbar(const TSharedRef<FUIComman
 		ToolBarBuilder.AddSeparator();
 		ToolBarBuilder.AddToolBarButton(FVisualLoggerCommands::Get().FreeCamera, NAME_None, LOCTEXT("FreeCamera", "Camera"), TAttribute<FText>(), FSlateIcon(FLogVisualizerStyle::Get().GetStyleSetName(), TEXT("Toolbar.Camera")));
 		ToolBarBuilder.AddSeparator();
-		ToolBarBuilder.AddToolBarButton(FVisualLoggerCommands::Get().Load, NAME_None, LOCTEXT("Load", "Load"), TAttribute<FText>(), FSlateIcon(FLogVisualizerStyle::Get().GetStyleSetName(), TEXT("Toolbar.Load")));
-		ToolBarBuilder.AddToolBarButton(FVisualLoggerCommands::Get().Save, NAME_None, LOCTEXT("Save", "Save"), TAttribute<FText>(), FSlateIcon(FLogVisualizerStyle::Get().GetStyleSetName(), TEXT("Toolbar.Save")));
+		ToolBarBuilder.AddToolBarButton(FVisualLoggerCommands::Get().LoadFromVLog, NAME_None, LOCTEXT("Load", "Load"), TAttribute<FText>(), FSlateIcon(FLogVisualizerStyle::Get().GetStyleSetName(), TEXT("Toolbar.Load")));
+		ToolBarBuilder.AddToolBarButton(FVisualLoggerCommands::Get().SaveToVLog, NAME_None, LOCTEXT("Save", "Save"), TAttribute<FText>(), FSlateIcon(FLogVisualizerStyle::Get().GetStyleSetName(), TEXT("Toolbar.Save")));
 		ToolBarBuilder.AddSeparator();
 		ToolBarBuilder.AddToolBarButton(FVisualLoggerCommands::Get().ToggleGraphs, NAME_None, LOCTEXT("ToggleGraphs", "Graphs"), TAttribute<FText>(), FSlateIcon(FLogVisualizerStyle::Get().GetStyleSetName(), TEXT("Toolbar.Graphs")));
 		ToolBarBuilder.AddToolBarButton(FVisualLoggerCommands::Get().ResetData, NAME_None, LOCTEXT("ResetData", "Clear"), TAttribute<FText>(), FSlateIcon(FLogVisualizerStyle::Get().GetStyleSetName(), TEXT("Toolbar.Remove")));
diff --git a/Engine/Source/Developer/LogVisualizer/Private/VisualLoggerCommands.h b/Engine/Source/Developer/LogVisualizer/Private/VisualLoggerCommands.h
--- a/Engine/Source/Developer/LogVisualizer/Private/VisualLoggerCommands.h
+++ b/Engine/Source/Developer/LogVisualizer/Private/VisualLoggerCommands.h
@@ -25,8 +25,8 @@ public:
 		UI_COMMAND(Pause, "Pause", "Stop the debugger", EUserInterfaceActionType::Button, FInputGesture());
 		UI_COMMAND(Resume, "Resume", "Stop the debugger", EUserInterfaceActionType::Button, FInputGesture());
 		UI_COMMAND(FreeCamera, "FreeCamera", "Enable free camera", EUserInterfaceActionType::ToggleButton, FInputGesture());
-		UI_COMMAND(Load, "Load", "Load external vlogs", EUserInterfaceActionType::Button, FInputGesture());
-		UI_COMMAND(Save, "Save", "Save selected data to vlog file", EUserInterfaceActionType::Button, FInputGesture());
+		UI_COMMAND(LoadFromVLog, "Load", "Load external vlogs", EUserInterfaceActionType::Button, FInputGesture());
+		UI_COMMAND(SaveToVLog, "Save", "Save selected data to vlog file", EUserInterfaceActionType::Button, FInputGesture());
 		UI_COMMAND(ToggleGraphs, "ToggleGraphs", "Toggle graphs visualization on/off", EUserInterfaceActionType::ToggleButton, FInputGesture());
 		UI_COMMAND(ResetData, "Clear", "Clear all data", EUserInterfaceActionType::Button, FInputGesture());
 	}
@@ -38,8 +38,8 @@ public:
 	TSharedPtr<FUICommandInfo> Pause;
 	TSharedPtr<FUICommandInfo> Resume;
 	TSharedPtr<FUICommandInfo> FreeCamera;
-	TSharedPtr<FUICommandInfo> Load;
-	TSharedPtr<FUICommandInfo> Save;
+	TSharedPtr<FUICommandInfo> LoadFromVLog;
+	TSharedPtr<FUICommandInfo> SaveToVLog;
 	TSharedPtr<FUICommandInfo> ToggleGraphs;
 	TSharedPtr<FUICommandInfo> ResetData;
 };
