Fixed issue with missing small toolbar icons in LogVisualizer.

diff --git a/Engine/Source/Developer/LogVisualizer/Private/LogVisualizerStyle.cpp b/Engine/Source/Developer/LogVisualizer/Private/LogVisualizerStyle.cpp
--- a/Engine/Source/Developer/LogVisualizer/Private/LogVisualizerStyle.cpp
+++ b/Engine/Source/Developer/LogVisualizer/Private/LogVisualizerStyle.cpp
@@ -83,6 +83,18 @@ TSharedRef< FSlateStyleSet > FLogVisualizerStyle::Create()
 		Style.Set("Toolbar.Load", new IMAGE_BRUSH("Icons/LV_Load", Icon40x40));
 		Style.Set("Toolbar.Remove", new IMAGE_BRUSH("Icons/Edit/icon_Edit_Delete_40x", Icon40x40));
 		Style.Set("Toolbar.Graphs", new IMAGE_BRUSH("Icons/icon_MatEd_Stats_40x", Icon40x40));
+		Style.Set("Toolbar.HideEmpty", new IMAGE_BRUSH("Icons/icon_levels_visible_40x", Icon40x40));
+
+		Style.Set("Toolbar.Pause.small", new IMAGE_BRUSH("Icons/icon_pause_40x", Icon20x20));
+		Style.Set("Toolbar.Resume.small", new IMAGE_BRUSH("Icons/icon_play_40x", Icon20x20));
+		Style.Set("Toolbar.Record.small", new IMAGE_BRUSH("Icons/LV_Record", Icon20x20));
+		Style.Set("Toolbar.Stop.small", new IMAGE_BRUSH("Icons/LV_Stop", Icon20x20));
+		Style.Set("Toolbar.Camera.small", new IMAGE_BRUSH("Icons/LV_Camera", Icon20x20));
+		Style.Set("Toolbar.Save.small", new IMAGE_BRUSH("Icons/LV_Save", Icon20x20));
+		Style.Set("Toolbar.Load.small", new IMAGE_BRUSH("Icons/LV_Load", Icon20x20));
+		Style.Set("Toolbar.Remove.small", new IMAGE_BRUSH("Icons/Edit/icon_Edit_Delete_40x", Icon20x20));
+		Style.Set("Toolbar.Graphs.small", new IMAGE_BRUSH("Icons/icon_MatEd_Stats_40x", Icon20x20));
+		Style.Set("Toolbar.HideEmpty.small", new IMAGE_BRUSH("Icons/icon_levels_visible_40x", Icon20x20));
 	}
 
 	// Filters
