Fixes for horizontal scroll bar in LogVisualizer tool
Categories/filters highlights from selected item on timeline  in LogVisualizer tool

diff --git a/Engine/Source/Developer/LogVisualizer/Private/SFilterWidget.cpp b/Engine/Source/Developer/LogVisualizer/Private/SFilterWidget.cpp
--- a/Engine/Source/Developer/LogVisualizer/Private/SFilterWidget.cpp
+++ b/Engine/Source/Developer/LogVisualizer/Private/SFilterWidget.cpp
@@ -16,6 +16,7 @@ void SFilterWidget::Construct(const FArguments& InArgs)
 	FilterName = InArgs._FilterName;
 	ColorCategory = InArgs._ColorCategory;
 	Verbosity = ELogVerbosity::All;
+	BorderBackgroundColor = FLinearColor(0.2f, 0.2f, 0.2f, 0.2f);
 
 	// Get the tooltip and color of the type represented by this filter
 	FilterColor = ColorCategory;
@@ -23,8 +24,8 @@ void SFilterWidget::Construct(const FArguments& InArgs)
 	ChildSlot
 		[
 			SNew(SBorder)
-			.Padding(0)
-			.BorderBackgroundColor(FLinearColor(0.2f, 0.2f, 0.2f, 0.2f))
+			.Padding(2)
+			.BorderBackgroundColor(this, &SFilterWidget::GetBorderBackgroundColor)
 			.BorderImage(FLogVisualizerStyle::Get().GetBrush("ContentBrowser.FilterButtonBorder"))
 			[
 				SAssignNew(ToggleButtonPtr, SFilterCheckBox)
diff --git a/Engine/Source/Developer/LogVisualizer/Private/SFilterWidget.h b/Engine/Source/Developer/LogVisualizer/Private/SFilterWidget.h
--- a/Engine/Source/Developer/LogVisualizer/Private/SFilterWidget.h
+++ b/Engine/Source/Developer/LogVisualizer/Private/SFilterWidget.h
@@ -14,7 +14,7 @@ public:
 
 	SLATE_BEGIN_ARGS(SFilterWidget){}
 		/** If this is an front end filter, this is the filter object */
-		SLATE_ARGUMENT(FName, FilterName)
+	SLATE_ARGUMENT(FName, FilterName)
 		/** Color selected for filter*/
 		SLATE_ARGUMENT(FLinearColor, ColorCategory)
 		/** Invoked when the filter toggled */
@@ -68,6 +68,8 @@ public:
 		return UnknownFilter;
 	}
 
+	void SetBorderBackgroundColor(FLinearColor InColor) { BorderBackgroundColor = InColor; }
+
 private:
 	/** Handler for when the filter checkbox is clicked */
 	void FilterToggled(ECheckBoxState NewState)
@@ -139,6 +141,11 @@ private:
 		return IsChecked() == ECheckBoxState::Checked ? ( IsHovered() ? ColorCategory * DimFactor : ColorCategory ) : ( IsHovered() ? FLinearColor::White : FLinearColor::White * DimFactor );
 	}
 
+	FSlateColor GetBorderBackgroundColor() const
+	{
+		return BorderBackgroundColor;
+	}
+
 private:
 	/** Invoked when the filter toggled */
 	FOnSimpleRequest OnFilterChanged;
@@ -169,4 +176,7 @@ private:
 
 	/** The color of the checkbox for this filter */
 	FLinearColor FilterColor;
+
+	/** Default color for border background */
+	FLinearColor BorderBackgroundColor;
 };
diff --git a/Engine/Source/Developer/LogVisualizer/Private/STimelineBar.cpp b/Engine/Source/Developer/LogVisualizer/Private/STimelineBar.cpp
--- a/Engine/Source/Developer/LogVisualizer/Private/STimelineBar.cpp
+++ b/Engine/Source/Developer/LogVisualizer/Private/STimelineBar.cpp
@@ -185,7 +185,7 @@ FReply STimelineBar::OnKeyDown(const FGeometry& MyGeometry, const FKeyEvent& InK
 	return ReturnValue;
 }
 
-uint32 STimelineBar::GetClosestItem(float Time) const
+int32 STimelineBar::GetClosestItem(float Time) const
 {
 	int32 BestItemIndex = INDEX_NONE;
 	float BestDistance = MAX_FLT;
@@ -266,7 +266,6 @@ void STimelineBar::OnSelect()
 {
 	FSlateApplication::Get().SetKeyboardFocus(SharedThis(this), EFocusCause::Navigation);
 	CurrentItemIndex = INDEX_NONE;
-	SnapScrubPosition(TimeSliderController->GetTimeSliderArgs().ScrubPosition.Get());
 }
 
 void STimelineBar::OnDeselect()
@@ -425,7 +424,7 @@ int32 STimelineBar::OnPaint(const FPaintArgs& Args, const FGeometry& AllottedGeo
 			);
 	}
 
-	uint32 BestItemIndex = GetClosestItem(LocalScrubPosition);
+	int32 BestItemIndex = GetClosestItem(LocalScrubPosition);
 
 	if (BestItemIndex != INDEX_NONE && TimelineOwner.Pin()->IsSelected())
 	{
diff --git a/Engine/Source/Developer/LogVisualizer/Private/STimelineBar.h b/Engine/Source/Developer/LogVisualizer/Private/STimelineBar.h
--- a/Engine/Source/Developer/LogVisualizer/Private/STimelineBar.h
+++ b/Engine/Source/Developer/LogVisualizer/Private/STimelineBar.h
@@ -28,7 +28,7 @@ public:
 	FVector2D ComputeDesiredSize() const;
 	void SnapScrubPosition(float ScrubPosition);
 	void SnapScrubPosition(int32 NewItemIndex);
-	uint32 GetClosestItem(float Time) const;
+	int32 GetClosestItem(float Time) const;
 
 	void OnSelect();
 	void OnDeselect();
diff --git a/Engine/Source/Developer/LogVisualizer/Private/STimelinesContainer.cpp b/Engine/Source/Developer/LogVisualizer/Private/STimelinesContainer.cpp
--- a/Engine/Source/Developer/LogVisualizer/Private/STimelinesContainer.cpp
+++ b/Engine/Source/Developer/LogVisualizer/Private/STimelinesContainer.cpp
@@ -211,6 +211,7 @@ FReply STimelinesContainer::OnKeyDown(const FGeometry& MyGeometry, const FKeyEve
 		SetSelectionState(NotSelectedOne.Pin(), true, true);
 		return FReply::Handled();
 	}
+#if 0 //disable movement between timelines for now
 	else if (InKeyEvent.GetKey() == EKeys::Up || InKeyEvent.GetKey() == EKeys::Down)
 	{
 		TSharedPtr<class STimeline> PreviousTimeline;
@@ -249,7 +250,7 @@ FReply STimelinesContainer::OnKeyDown(const FGeometry& MyGeometry, const FKeyEve
 		}
 		return FReply::Handled();
 	}
-
+#endif
 	return FReply::Unhandled();
 }
 
diff --git a/Engine/Source/Developer/LogVisualizer/Private/SVisualLogger.cpp b/Engine/Source/Developer/LogVisualizer/Private/SVisualLogger.cpp
--- a/Engine/Source/Developer/LogVisualizer/Private/SVisualLogger.cpp
+++ b/Engine/Source/Developer/LogVisualizer/Private/SVisualLogger.cpp
@@ -747,11 +747,6 @@ void SVisualLogger::HandleSaveCommandExecute()
 
 void SVisualLogger::ResetData()
 {
-	if (ULogVisualizerSettings::StaticClass()->GetDefaultObject<ULogVisualizerSettings>()->bResetDataWithNewSession)
-	{
-		return;
-	}
-
 	if (MainView.IsValid())
 	{
 		MainView->ResetData();
@@ -766,7 +761,10 @@ void SVisualLogger::ResetData()
 void SVisualLogger::OnNewWorld(UWorld* NewWorld)
 {
 	InternalDevice->SerLastWorld(NewWorld);
-	ResetData();
+	if (ULogVisualizerSettings::StaticClass()->GetDefaultObject<ULogVisualizerSettings>()->bResetDataWithNewSession)
+	{
+		ResetData();
+	}
 }
 
 void SVisualLogger::OnNewLogEntry(const FVisualLogDevice::FVisualLogEntryItem& Entry)
@@ -800,6 +798,7 @@ void SVisualLogger::OnItemSelectionChanged(const FVisualLogDevice::FVisualLogEnt
 	LogsList->OnItemSelectionChanged(EntryItem);
 	StatusView->OnItemSelectionChanged(EntryItem);
 	VisualLoggerCanvasRenderer->OnItemSelectionChanged(EntryItem.Entry);
+	VisualLoggerFilters->OnItemSelectionChanged(EntryItem.Entry);
 	AVisualLoggerRenderingActor* HelperActor = Cast<AVisualLoggerRenderingActor>(VisualLoggerInterface->GetVisualLoggerHelperActor());
 	if (HelperActor)
 	{
diff --git a/Engine/Source/Developer/LogVisualizer/Private/SVisualLoggerFilters.cpp b/Engine/Source/Developer/LogVisualizer/Private/SVisualLoggerFilters.cpp
--- a/Engine/Source/Developer/LogVisualizer/Private/SVisualLoggerFilters.cpp
+++ b/Engine/Source/Developer/LogVisualizer/Private/SVisualLoggerFilters.cpp
@@ -462,4 +462,22 @@ void SVisualLoggerFilters::OnFiltersSearchChanged(const FText& Filter)
 
 }
 
+void SVisualLoggerFilters::OnItemSelectionChanged(const struct FVisualLogEntry& EntryItem)
+{
+	TArray<FVisualLoggerCategoryVerbosityPair> Categories;
+	FVisualLoggerHelpers::GetCategories(EntryItem, Categories);
+	for (int32 Index = 0; Index < Filters.Num(); ++Index)
+	{
+		SFilterWidget& Filter = Filters[Index].Get();
+		Filter.SetBorderBackgroundColor(FLinearColor(0.2f, 0.2f, 0.2f, 0.2f));
+		for (const FVisualLoggerCategoryVerbosityPair& Category : Categories)
+		{
+			if (Filter.GetFilterName() == Category.CategoryName)
+			{
+				Filter.SetBorderBackgroundColor(FLinearColor(0.3f, 0.3f, 0.3f, 0.8f));
+			}
+		}
+	}
+}
+
 #undef LOCTEXT_NAMESPACE
diff --git a/Engine/Source/Developer/LogVisualizer/Private/SVisualLoggerFilters.h b/Engine/Source/Developer/LogVisualizer/Private/SVisualLoggerFilters.h
--- a/Engine/Source/Developer/LogVisualizer/Private/SVisualLoggerFilters.h
+++ b/Engine/Source/Developer/LogVisualizer/Private/SVisualLoggerFilters.h
@@ -45,6 +45,7 @@ public:
 	bool IsGraphFilterCategoryInUse(FName MenuCategory) const;
 	void OnFiltersChanged();
 	void ResetData();
+	void OnItemSelectionChanged(const struct FVisualLogEntry& EntryItem);
 
 	TSharedRef<SWidget> MakeGraphsFilterMenu();
 	void CreateFiltersMenuCategoryForGraph(FMenuBuilder& MenuBuilder, FName MenuCategory) const;
diff --git a/Engine/Source/Developer/LogVisualizer/Private/SVisualLoggerToolbar.cpp b/Engine/Source/Developer/LogVisualizer/Private/SVisualLoggerToolbar.cpp
--- a/Engine/Source/Developer/LogVisualizer/Private/SVisualLoggerToolbar.cpp
+++ b/Engine/Source/Developer/LogVisualizer/Private/SVisualLoggerToolbar.cpp
@@ -32,8 +32,9 @@ TSharedRef<SWidget> SVisualLoggerToolbar::MakeToolbar(const TSharedRef<FUIComman
 		ToolBarBuilder.AddToolBarButton(FVisualLoggerCommands::Get().LoadFromVLog, NAME_None, LOCTEXT("Load", "Load"), TAttribute<FText>(), FSlateIcon(FLogVisualizerStyle::Get().GetStyleSetName(), TEXT("Toolbar.Load")));
 		ToolBarBuilder.AddToolBarButton(FVisualLoggerCommands::Get().SaveToVLog, NAME_None, LOCTEXT("Save", "Save"), TAttribute<FText>(), FSlateIcon(FLogVisualizerStyle::Get().GetStyleSetName(), TEXT("Toolbar.Save")));
 		ToolBarBuilder.AddSeparator();
-		ToolBarBuilder.AddToolBarButton(FVisualLoggerCommands::Get().ToggleGraphs, NAME_None, LOCTEXT("ToggleGraphs", "Graphs"), TAttribute<FText>(), FSlateIcon(FLogVisualizerStyle::Get().GetStyleSetName(), TEXT("Toolbar.Graphs")));
 		ToolBarBuilder.AddToolBarButton(FVisualLoggerCommands::Get().ResetData, NAME_None, LOCTEXT("ResetData", "Clear"), TAttribute<FText>(), FSlateIcon(FLogVisualizerStyle::Get().GetStyleSetName(), TEXT("Toolbar.Remove")));
+		ToolBarBuilder.AddToolBarButton(FVisualLoggerCommands::Get().ToggleGraphs, NAME_None, LOCTEXT("ToggleGraphs", "Graphs"), TAttribute<FText>(), FSlateIcon(FLogVisualizerStyle::Get().GetStyleSetName(), TEXT("Toolbar.Graphs")));
+		//ToolBarBuilder.AddToolBarButton(FVisualLoggerCommands::Get().HideEmpty, NAME_None, LOCTEXT("HideEmpty", "Hide Empty"), TAttribute<FText>(), FSlateIcon(FLogVisualizerStyle::Get().GetStyleSetName(), TEXT("Toolbar.HideEmpty")));
 	}
 
 	return ToolBarBuilder.MakeWidget();
diff --git a/Engine/Source/Developer/LogVisualizer/Private/VisualLoggerCommands.h b/Engine/Source/Developer/LogVisualizer/Private/VisualLoggerCommands.h
--- a/Engine/Source/Developer/LogVisualizer/Private/VisualLoggerCommands.h
+++ b/Engine/Source/Developer/LogVisualizer/Private/VisualLoggerCommands.h
@@ -29,6 +29,7 @@ public:
 		UI_COMMAND(SaveToVLog, "Save", "Save selected data to vlog file", EUserInterfaceActionType::Button, FInputGesture());
 		UI_COMMAND(ToggleGraphs, "ToggleGraphs", "Toggle graphs visualization on/off", EUserInterfaceActionType::ToggleButton, FInputGesture());
 		UI_COMMAND(ResetData, "Clear", "Clear all data", EUserInterfaceActionType::Button, FInputGesture());
+		UI_COMMAND(HideEmpty, "HideEmpty", "Hide logs without text information", EUserInterfaceActionType::ToggleButton, FInputGesture());
 	}
 
 public:
@@ -42,6 +43,7 @@ public:
 	TSharedPtr<FUICommandInfo> SaveToVLog;
 	TSharedPtr<FUICommandInfo> ToggleGraphs;
 	TSharedPtr<FUICommandInfo> ResetData;
+	TSharedPtr<FUICommandInfo> HideEmpty;
 };
 
 
