Added extreme values for graphs in LogVisualizer (disabled by default and it can be enabled in settings)

diff --git a/Engine/Source/Developer/LogVisualizer/Private/LogVisualizerSettings.cpp b/Engine/Source/Developer/LogVisualizer/Private/LogVisualizerSettings.cpp
--- a/Engine/Source/Developer/LogVisualizer/Private/LogVisualizerSettings.cpp
+++ b/Engine/Source/Developer/LogVisualizer/Private/LogVisualizerSettings.cpp
@@ -17,6 +17,7 @@ ULogVisualizerSettings::ULogVisualizerSettings(const FObjectInitializer& ObjectI
 	bSearchInsideLogs = true;
 	GraphsBackgroundColor = FColor(0, 0, 0, 70);
 	bResetDataWithNewSession = false;
+	bDrawExtremesOnGraphs = false;
 }
 
 #if WITH_EDITOR
diff --git a/Engine/Source/Developer/LogVisualizer/Private/VisualLoggerCanvasRenderer.cpp b/Engine/Source/Developer/LogVisualizer/Private/VisualLoggerCanvasRenderer.cpp
--- a/Engine/Source/Developer/LogVisualizer/Private/VisualLoggerCanvasRenderer.cpp
+++ b/Engine/Source/Developer/LogVisualizer/Private/VisualLoggerCanvasRenderer.cpp
@@ -112,6 +112,7 @@ void FVisualLoggerCanvasRenderer::DrawHistogramGraphs(class UCanvas* Canvas, cla
 	struct FGraphLineData
 	{
 		FName DataName;
+		FVector2D LeftExtreme, RightExtreme;
 		TArray<FVector2D> Samples;
 	};
 	typedef TMap<FName, FGraphLineData > FGraphLines;
@@ -141,16 +142,21 @@ void FVisualLoggerCanvasRenderer::DrawHistogramGraphs(class UCanvas* Canvas, cla
 
 	const TArray<FVisualLogDevice::FVisualLogEntryItem> &ObjectItems = CurrentTimeLine.Pin()->GetEntries();
 	int32 ColorIndex = 0;
+	int32 LeftSideOutsideIndex = INDEX_NONE;
+	int32 RightSideOutsideIndex = INDEX_NONE;
+
 	for (int32 EntryIndex = 0; EntryIndex < ObjectItems.Num(); ++EntryIndex)
 	{
 		const FVisualLogEntry* CurrentEntry = &(ObjectItems[EntryIndex].Entry);
 		if (CurrentEntry->TimeStamp < TimeStampWindow.X)
 		{
+			LeftSideOutsideIndex = EntryIndex;
 			continue;
 		}
 
 		if (CurrentEntry->TimeStamp > TimeStampWindow.Y)
 		{
+			RightSideOutsideIndex = EntryIndex;
 			break;
 		}
 
@@ -164,10 +170,7 @@ void FVisualLoggerCanvasRenderer::DrawHistogramGraphs(class UCanvas* Canvas, cla
 			const FName CurrentDataName = CurrentSample.DataName;
 
 			const bool bIsValidByFilter = VisualLoggerInterface.Pin()->IsValidCategory(CurrentSample.GraphName.ToString(), CurrentSample.DataName.ToString(), ELogVerbosity::All);
-			const bool bCurrentDataNamePassed = true;// !bHistogramGraphsFilter || (QuickFilterText.Len() == 0 || CurrentDataName.ToString().Find(QuickFilterText) != INDEX_NONE);
-
-
-			if (bIsValidByFilter /*&& bCurrentCategoryPassed*/ && bCurrentDataNamePassed)
+			if (bIsValidByFilter)
 			{
 				FGraphData &GraphData = CollectedGraphs.FindOrAdd(CurrentSample.GraphName);
 				FGraphLineData &LineData = GraphData.GraphLines.FindOrAdd(CurrentSample.DataName);
@@ -182,6 +185,33 @@ void FVisualLoggerCanvasRenderer::DrawHistogramGraphs(class UCanvas* Canvas, cla
 			}
 		}
 	}
+	
+	const int32 ExtremeValueIndexes[] = { LeftSideOutsideIndex != INDEX_NONE ? LeftSideOutsideIndex : 0, RightSideOutsideIndex != INDEX_NONE ? RightSideOutsideIndex : ObjectItems.Num() - 1 };
+	for (int32 ObjectIndex = 0; ObjectIndex < 2; ++ObjectIndex)
+	{
+		const FVisualLogEntry* CurrentEntry = &(ObjectItems[ExtremeValueIndexes[ObjectIndex]].Entry);
+		const int32 SamplesNum = CurrentEntry->HistogramSamples.Num();
+		for (int32 SampleIndex = 0; SampleIndex < SamplesNum; ++SampleIndex)
+		{
+			FVisualLogHistogramSample CurrentSample = CurrentEntry->HistogramSamples[SampleIndex];
+
+			const FName CurrentCategory = CurrentSample.Category;
+			const FName CurrentGraphName = CurrentSample.GraphName;
+			const FName CurrentDataName = CurrentSample.DataName;
+
+			const bool bIsValidByFilter = VisualLoggerInterface.Pin()->IsValidCategory(CurrentSample.GraphName.ToString(), CurrentSample.DataName.ToString(), ELogVerbosity::All);
+			if (bIsValidByFilter)
+			{
+				FGraphData &GraphData = CollectedGraphs.FindOrAdd(CurrentSample.GraphName);
+				FGraphLineData &LineData = GraphData.GraphLines.FindOrAdd(CurrentSample.DataName);
+				LineData.DataName = CurrentSample.DataName;
+				if (ObjectIndex == 0)
+					LineData.LeftExtreme = CurrentSample.SampleValue;
+				else
+					LineData.RightExtreme = CurrentSample.SampleValue;
+			}
+		}
+	}
 
 	const float GoldenRatioConjugate = 0.618033988749895f;
 	int32 GraphIndex = 0;
@@ -211,6 +241,7 @@ void FVisualLoggerCanvasRenderer::DrawHistogramGraphs(class UCanvas* Canvas, cla
 		int32 GraphIndex = 0;
 		int32 CurrentColumn = 0;
 		int32 CurrentRow = 0;
+		bool bDrawExtremesOnGraphs = ULogVisualizerSettings::StaticClass()->GetDefaultObject<ULogVisualizerSettings>()->bDrawExtremesOnGraphs;
 		for (auto It(CollectedGraphs.CreateIterator()); It; ++It)
 		{
 			TWeakObjectPtr<UReporterGraph> HistogramGraph = Canvas->GetReporterGraph();
@@ -222,7 +253,7 @@ void FVisualLoggerCanvasRenderer::DrawHistogramGraphs(class UCanvas* Canvas, cla
 			int32 LineIndex = 0;
 			UFont* Font = GEngine->GetSmallFont();
 			int32 MaxStringSize = 0;
-			float Hue = 0;//StartGoldenRatio[GraphIndex++];
+			float Hue = 0;
 
 			auto& CategoriesForGraph = UsedGraphCategories.FindOrAdd(It->Key.ToString());
 
@@ -244,6 +275,8 @@ void FVisualLoggerCanvasRenderer::DrawHistogramGraphs(class UCanvas* Canvas, cla
 				HistogramGraph->GetGraphLine(LineIndex)->Color = FLinearColor::FGetHSV(Hue * 255, 0, 244);
 				HistogramGraph->GetGraphLine(LineIndex)->LineName = DataName;
 				HistogramGraph->GetGraphLine(LineIndex)->Data.Append(LinesIt->Value.Samples);
+				HistogramGraph->GetGraphLine(LineIndex)->LeftExtreme = LinesIt->Value.LeftExtreme;
+				HistogramGraph->GetGraphLine(LineIndex)->RightExtreme = LinesIt->Value.RightExtreme;
 
 				int32 DummyY, CurrentX;
 				StringSize(Font, CurrentX, DummyY, *LinesIt->Value.DataName.ToString());
@@ -267,6 +300,7 @@ void FVisualLoggerCanvasRenderer::DrawHistogramGraphs(class UCanvas* Canvas, cla
 			HistogramGraph->SetBackgroundColor(GraphsBackgroundColor);
 			HistogramGraph->SetLegendPosition(/*bShowHistogramLabelsOutside*/ false ? ELegendPosition::Outside : ELegendPosition::Inside);
 			HistogramGraph->OffsetDataSets(/*bOffsetDataSet*/false);
+			HistogramGraph->DrawExtremesOnGraph(bDrawExtremesOnGraphs);
 			HistogramGraph->bVisible = true;
 			HistogramGraph->Draw(Canvas);
 
diff --git a/Engine/Source/Developer/LogVisualizer/Public/LogVisualizerSettings.h b/Engine/Source/Developer/LogVisualizer/Public/LogVisualizerSettings.h
--- a/Engine/Source/Developer/LogVisualizer/Public/LogVisualizerSettings.h
+++ b/Engine/Source/Developer/LogVisualizer/Public/LogVisualizerSettings.h
@@ -72,11 +72,14 @@ public:
 	UPROPERTY(config)
 	TArray<FFiltersPreset> FilterPresets;
 
-
 	/**Whether to store all filter settings on exit*/
 	UPROPERTY(EditAnywhere, config, Category = "VisualLogger")
 	bool bPresistentFilters;
 
+	/** Whether to extreme values on graph (data has to be provided for extreme values) */
+	UPROPERTY(EditAnywhere, config, Category = "VisualLogger")
+	bool bDrawExtremesOnGraphs;
+
 	UPROPERTY(config)
 	FFiltersPreset PresistentPresets;
 
diff --git a/Engine/Source/Runtime/Engine/Classes/Debug/ReporterGraph.h b/Engine/Source/Runtime/Engine/Classes/Debug/ReporterGraph.h
--- a/Engine/Source/Runtime/Engine/Classes/Debug/ReporterGraph.h
+++ b/Engine/Source/Runtime/Engine/Classes/Debug/ReporterGraph.h
@@ -82,6 +82,12 @@ struct FGraphLine
 	/** The colour to graph this data with */
 	FLinearColor Color;
 
+	/** Left extreme value */
+	FVector2D LeftExtreme;
+
+	/** Right extreme value */
+	FVector2D RightExtreme;
+
 	/** The event name */
 	UPROPERTY(EditAnywhere, BlueprintReadWrite, Category=FCombatEventData)
 	FString LineName;
@@ -193,7 +199,8 @@ public:
 
 	/** setter to enable or disable cursor for line graphs */
 	void DrawCursorOnGraph(bool InDrawCursorOnGraph) { bDrawCursorOnGraph = InDrawCursorOnGraph;  }
-protected:
+
+	void DrawExtremesOnGraph(bool InDrawExtremes) { bDrawExtremes = InDrawExtremes; }
 
 	/** Draw background under graph
 	* @param Canvas The canvas to draw on
@@ -291,5 +298,8 @@ protected:
 
 	/** if set it enables cursor for line graphs */
 	int32 bDrawCursorOnGraph : 1;
+
+	/** draw extreams on vertical axes */
+	int32 bDrawExtremes : 1;
 };
 
diff --git a/Engine/Source/Runtime/Engine/Private/Debug/ReporterGraph.cpp b/Engine/Source/Runtime/Engine/Private/Debug/ReporterGraph.cpp
--- a/Engine/Source/Runtime/Engine/Private/Debug/ReporterGraph.cpp
+++ b/Engine/Source/Runtime/Engine/Private/Debug/ReporterGraph.cpp
@@ -303,6 +303,34 @@ void UReporterGraph::DrawAxis(UCanvas* Canvas, FVector2D Start, FVector2D End, f
 			NotchLocation.X += NotchDelta;
 		}
 	}
+
+	if (bIsVerticalAxis && bDrawExtremes)
+	{
+		for (int32 i = 0; i < CurrentData.Num(); i++)
+		{
+
+			if (CurrentData[i].Data.Num() == 0)
+				continue;
+
+			FVector2D DataStart = CurrentData[i].LeftExtreme;
+			DataStart.X = Start.X;
+			FVector2D TextPos = ToScreenSpace(DataToNormalized(DataStart), Canvas) /*+ DrawOffset*/;
+
+			FString Text = FString::Printf(TEXT("%.2f"), CurrentData[i].LeftExtreme.Y);
+			int32 StringSizeX, StringSizeY;
+			StringSize(Font, StringSizeX, StringSizeY, *Text);
+			Canvas->Canvas->DrawShadowedString(TextPos.X - StringSizeX * 0.5f, TextPos.Y + (AxisStyle == EGraphAxisStyle::Grid ? +5 : -NotchLength.Y * Canvas->SizeY), *Text, Font, CurrentData[i].Color);
+
+			FVector2D DataEnd = CurrentData[i].Data[CurrentData[i].Data.Num()-1];
+			DataEnd = DataToNormalized(DataEnd);
+			DataEnd.X = GraphScreenSize.Max.X;
+			TextPos = ToScreenSpace(DataEnd, Canvas) /*+ DrawOffset*/;
+
+			Text = FString::Printf(TEXT("%.2f"), CurrentData[i].RightExtreme.Y);
+			StringSize(Font, StringSizeX, StringSizeY, *Text);
+			Canvas->Canvas->DrawShadowedString(TextPos.X - StringSizeX * 0.5f, TextPos.Y + (AxisStyle == EGraphAxisStyle::Grid ? +5 : -NotchLength.Y * Canvas->SizeY), *Text, Font, CurrentData[i].Color);
+		}
+	}
 }
 
 void UReporterGraph::DrawThresholds(UCanvas* Canvas)
