Fixed UE-7430 "Visual Logger Toolbar has extra divider"

diff --git a/Engine/Source/Developer/LogVisualizer/Private/SVisualLoggerToolbar.cpp b/Engine/Source/Developer/LogVisualizer/Private/SVisualLoggerToolbar.cpp
--- a/Engine/Source/Developer/LogVisualizer/Private/SVisualLoggerToolbar.cpp
+++ b/Engine/Source/Developer/LogVisualizer/Private/SVisualLoggerToolbar.cpp
@@ -26,15 +26,14 @@ TSharedRef<SWidget> SVisualLoggerToolbar::MakeToolbar(const TSharedRef<FUIComman
 		ToolBarBuilder.AddToolBarButton(FVisualLoggerCommands::Get().Pause, NAME_None, LOCTEXT("PauseGame", "Pause"), TAttribute<FText>(), FSlateIcon(FLogVisualizerStyle::Get().GetStyleSetName(), TEXT("Toolbar.Pause")));
 		ToolBarBuilder.AddToolBarButton(FVisualLoggerCommands::Get().Resume, NAME_None, LOCTEXT("ResumeGame", "Resume"), TAttribute<FText>(), FSlateIcon(FLogVisualizerStyle::Get().GetStyleSetName(), TEXT("Toolbar.Resume")));
 
-		ToolBarBuilder.AddSeparator();
-		ToolBarBuilder.AddToolBarButton(FVisualLoggerCommands::Get().FreeCamera, NAME_None, LOCTEXT("FreeCamera", "Camera"), TAttribute<FText>(), FSlateIcon(FLogVisualizerStyle::Get().GetStyleSetName(), TEXT("Toolbar.Camera")));
 		ToolBarBuilder.AddSeparator();
 		ToolBarBuilder.AddToolBarButton(FVisualLoggerCommands::Get().LoadFromVLog, NAME_None, LOCTEXT("Load", "Load"), TAttribute<FText>(), FSlateIcon(FLogVisualizerStyle::Get().GetStyleSetName(), TEXT("Toolbar.Load")));
 		ToolBarBuilder.AddToolBarButton(FVisualLoggerCommands::Get().SaveToVLog, NAME_None, LOCTEXT("Save", "Save"), TAttribute<FText>(), FSlateIcon(FLogVisualizerStyle::Get().GetStyleSetName(), TEXT("Toolbar.Save")));
+
 		ToolBarBuilder.AddSeparator();
+		ToolBarBuilder.AddToolBarButton(FVisualLoggerCommands::Get().FreeCamera, NAME_None, LOCTEXT("FreeCamera", "Camera"), TAttribute<FText>(), FSlateIcon(FLogVisualizerStyle::Get().GetStyleSetName(), TEXT("Toolbar.Camera")));
 		ToolBarBuilder.AddToolBarButton(FVisualLoggerCommands::Get().ResetData, NAME_None, LOCTEXT("ResetData", "Clear"), TAttribute<FText>(), FSlateIcon(FLogVisualizerStyle::Get().GetStyleSetName(), TEXT("Toolbar.Remove")));
 		ToolBarBuilder.AddToolBarButton(FVisualLoggerCommands::Get().ToggleGraphs, NAME_None, LOCTEXT("ToggleGraphs", "Graphs"), TAttribute<FText>(), FSlateIcon(FLogVisualizerStyle::Get().GetStyleSetName(), TEXT("Toolbar.Graphs")));
-		//ToolBarBuilder.AddToolBarButton(FVisualLoggerCommands::Get().HideEmpty, NAME_None, LOCTEXT("HideEmpty", "Hide Empty"), TAttribute<FText>(), FSlateIcon(FLogVisualizerStyle::Get().GetStyleSetName(), TEXT("Toolbar.HideEmpty")));
 	}
 
 	return ToolBarBuilder.MakeWidget();
