Added red and yellow colors for text lines with error or warning in LogVisualizer tool.

diff --git a/Engine/Source/Developer/LogVisualizer/Private/SVisualLoggerLogsList.cpp b/Engine/Source/Developer/LogVisualizer/Private/SVisualLoggerLogsList.cpp
--- a/Engine/Source/Developer/LogVisualizer/Private/SVisualLoggerLogsList.cpp
+++ b/Engine/Source/Developer/LogVisualizer/Private/SVisualLoggerLogsList.cpp
@@ -110,7 +110,7 @@ TSharedRef<ITableRow> SVisualLoggerLogsList::LogEntryLinesGenerateRow(TSharedPtr
 			.Padding(FMargin(5.0f, 0.0f))
 			[
 				SNew(STextBlock)
-				.ColorAndOpacity(FSlateColor(FLinearColor::Gray))
+				.ColorAndOpacity(FSlateColor(Item->Verbosity == ELogVerbosity::Error ? FLinearColor::Red : (Item->Verbosity == ELogVerbosity::Warning ? FLinearColor::Yellow : FLinearColor::Gray)))
 				.Text(FText::FromString(FString(TEXT("(")) + FString(FOutputDevice::VerbosityToString(Item->Verbosity)) + FString(TEXT(")"))))
 			]
 			+ SHorizontalBox::Slot()
@@ -118,6 +118,7 @@ TSharedRef<ITableRow> SVisualLoggerLogsList::LogEntryLinesGenerateRow(TSharedPtr
 			[
 				SNew(STextBlock)
 				.AutoWrapText(true)
+				.ColorAndOpacity(FSlateColor(Item->Verbosity == ELogVerbosity::Error ? FLinearColor::Red : (Item->Verbosity == ELogVerbosity::Warning ? FLinearColor::Yellow : FLinearColor::Gray)))
 				.Text(FText::FromString(Item->Line))
 				.HighlightText(this, &SVisualLoggerLogsList::GetFilterText)
 			]
