Bunch of internal changes to LogVisualize and added module interface to controll cursor location on timelines.

diff --git a/Engine/Source/Developer/LogVisualizer/Private/LogVisualizer.cpp b/Engine/Source/Developer/LogVisualizer/Private/LogVisualizer.cpp
new file mode 100644
--- /dev/null
+++ b/Engine/Source/Developer/LogVisualizer/Private/LogVisualizer.cpp
@@ -0,0 +1,192 @@
+// Copyright 1998-2015 Epic Games, Inc. All Rights Reserved.
+
+#include "LogVisualizer.h"
+#include "TimeSliderController.h"
+#include "VisualLoggerRenderingActor.h"
+#include "STimelinesContainer.h"
+#if WITH_EDITOR
+#	include "Editor/UnrealEd/Public/EditorComponents.h"
+#	include "Editor/UnrealEd/Public/EditorReimportHandler.h"
+#	include "Editor/UnrealEd/Public/TexAlignTools.h"
+#	include "Editor/UnrealEd/Public/TickableEditorObject.h"
+#	include "UnrealEdClasses.h"
+#	include "Editor/UnrealEd/Public/Editor.h"
+#	include "Editor/UnrealEd/Public/EditorViewportClient.h"
+#endif
+
+TSharedPtr< struct FLogVisualizer > FLogVisualizer::StaticInstance;
+FColor FLogVisualizerColorPalette[] = {
+	FColor(0xff00A480),
+	FColorList::Aquamarine,
+	FColorList::Cyan,
+	FColorList::Brown,
+	FColorList::Green,
+	FColorList::Orange,
+	FColorList::Magenta,
+	FColorList::BrightGold,
+	FColorList::NeonBlue,
+	FColorList::MediumSlateBlue,
+	FColorList::SpicyPink,
+	FColorList::SpringGreen,
+	FColorList::SteelBlue,
+	FColorList::SummerSky,
+	FColorList::Violet,
+	FColorList::VioletRed,
+	FColorList::YellowGreen,
+	FColor(0xff62E200),
+	FColor(0xff1F7B67),
+	FColor(0xff62AA2A),
+	FColor(0xff70227E),
+	FColor(0xff006B53),
+	FColor(0xff409300),
+	FColor(0xff5D016D),
+	FColor(0xff34D2AF),
+	FColor(0xff8BF13C),
+	FColor(0xffBC38D3),
+	FColor(0xff5ED2B8),
+	FColor(0xffA6F16C),
+	FColor(0xffC262D3),
+	FColor(0xff0F4FA8),
+	FColor(0xff00AE68),
+	FColor(0xffDC0055),
+	FColor(0xff284C7E),
+	FColor(0xff21825B),
+	FColor(0xffA52959),
+	FColor(0xff05316D),
+	FColor(0xff007143),
+	FColor(0xff8F0037),
+	FColor(0xff4380D3),
+	FColor(0xff36D695),
+	FColor(0xffEE3B80),
+	FColor(0xff6996D3),
+	FColor(0xff60D6A7),
+	FColor(0xffEE6B9E)
+};
+
+
+void FLogVisualizer::Initialize()
+{
+	StaticInstance = MakeShareable(new FLogVisualizer);
+	Get().TimeSliderController = MakeShareable(new FSequencerTimeSliderController(FVisualLoggerTimeSliderArgs()));
+
+}
+
+void FLogVisualizer::Shutdown()
+{
+	StaticInstance.Reset();
+}
+
+FLogVisualizer& FLogVisualizer::Get()
+{
+	return *StaticInstance;
+}
+
+void FLogVisualizer::Goto(float Timestamp, FName LogOwner)
+{
+	if (CurrentVisualizer.IsValid())
+	{
+		TArray<TSharedPtr<class STimeline> > TimeLines;
+		CurrentVisualizer.Pin()->GetTimelines(TimeLines, false);
+		for (TSharedPtr<class STimeline> Timeline : TimeLines)
+		{
+			if (Timeline->GetName() == LogOwner)
+			{
+				Timeline->GetOwner()->SetSelectionState(Timeline, true, true);
+				break;
+			}
+		}
+	}
+
+	if (CurrentTimeLine.IsValid())
+	{
+		CurrentTimeLine.Pin()->Goto(Timestamp);
+	}
+}
+
+void FLogVisualizer::GotoNextItem()
+{
+	if (CurrentTimeLine.IsValid())
+	{
+		CurrentTimeLine.Pin()->GotoNextItem();
+	}
+}
+
+void FLogVisualizer::GotoPreviousItem()
+{
+	if (CurrentTimeLine.IsValid())
+	{
+		CurrentTimeLine.Pin()->GotoPreviousItem();
+	}
+}
+
+FLinearColor FLogVisualizer::GetColorForCategory(int32 Index) const
+{
+	if (Index >= 0 && Index < sizeof(FLogVisualizerColorPalette) / sizeof(FLogVisualizerColorPalette[0]))
+	{
+		return FLogVisualizerColorPalette[Index];
+	}
+
+	static bool bReateColorList = false;
+	static FColorList StaticColor;
+	if (!bReateColorList)
+	{
+		bReateColorList = true;
+		StaticColor.CreateColorMap();
+	}
+	return StaticColor.GetFColorByIndex(Index);
+}
+
+FLinearColor FLogVisualizer::GetColorForCategory(const FString& InFilterName) const
+{
+	static TArray<FString> Filters;
+	int32 CategoryIndex = Filters.Find(InFilterName);
+	if (CategoryIndex == INDEX_NONE)
+	{
+		CategoryIndex = Filters.Add(InFilterName);
+	}
+
+	return GetColorForCategory(CategoryIndex);
+}
+
+UWorld* FLogVisualizer::GetWorld(UObject* OptionalObject)
+{
+	UWorld* World = OptionalObject != nullptr ? GEngine->GetWorldFromContextObject(OptionalObject, false) : nullptr;
+#if WITH_EDITOR
+	if (!World && GIsEditor)
+	{
+		UEditorEngine *EEngine = Cast<UEditorEngine>(GEngine);
+		// lets use PlayWorld during PIE/Simulate and regular world from editor otherwise, to draw debug information
+		World = EEngine != nullptr && EEngine->PlayWorld != nullptr ? EEngine->PlayWorld : EEngine->GetEditorWorldContext().World();
+
+	}
+	else 
+#endif
+	if (!World && !GIsEditor)
+	{
+
+		World = GEngine->GetWorld();
+	}
+
+	if (World == nullptr)
+	{
+		World = GWorld;
+	}
+
+	return World;
+}
+
+AActor* FLogVisualizer::GetVisualLoggerHelperActor()
+{
+	UWorld* World = FLogVisualizer::Get().GetWorld();
+	for (TActorIterator<AVisualLoggerRenderingActor> It(World); It; ++It)
+	{
+		return *It;
+	}
+
+	FActorSpawnParameters SpawnInfo;
+	SpawnInfo.bNoCollisionFail = true;
+	SpawnInfo.Name = *FString::Printf(TEXT("VisualLoggerRenderingActor"));
+	AActor* HelperActor = World->SpawnActor<AVisualLoggerRenderingActor>(SpawnInfo);
+
+	return HelperActor;
+}
diff --git a/Engine/Source/Developer/LogVisualizer/Private/LogVisualizer.h b/Engine/Source/Developer/LogVisualizer/Private/LogVisualizer.h
--- a/Engine/Source/Developer/LogVisualizer/Private/LogVisualizer.h
+++ b/Engine/Source/Developer/LogVisualizer/Private/LogVisualizer.h
@@ -12,6 +12,7 @@
 
 /* Private includes
 *****************************************************************************/
+#include "ILogVisualizer.h"
 #include "VisualLogger/VisualLogger.h"
 #include "LogVisualizerSettings.h"
 #include "LogVisualizerSessionSettings.h"
@@ -39,22 +40,43 @@ struct IVisualLoggerInterface
 	virtual bool HasValidCategories(TArray<FVisualLoggerCategoryVerbosityPair> Categories) = 0;
 	virtual bool IsValidCategory(const FString& InCategoryName, TEnumAsByte<ELogVerbosity::Type> Verbosity = ELogVerbosity::All) = 0;
 	virtual bool IsValidCategory(const FString& InGraphName, const FString& InDataName, TEnumAsByte<ELogVerbosity::Type> Verbosity = ELogVerbosity::All) = 0;
-	virtual FLinearColor GetCategoryColor(FName Category) = 0;
-	virtual UWorld* GetWorld() const = 0;
-	virtual class AActor* GetVisualLoggerHelperActor() = 0;
+};
+
+class FSequencerTimeSliderController;
+struct LOGVISUALIZER_API FLogVisualizer
+{
+	/** LogVisualizer interface*/
+	void Goto(float Timestamp, FName LogOwner = NAME_None);
+	void GotoNextItem();
+	void GotoPreviousItem();
+
+	FLinearColor GetColorForCategory(int32 Index) const;
+	FLinearColor GetColorForCategory(const FString& InFilterName) const;
+	TSharedPtr<FSequencerTimeSliderController> GetTimeSliderController() { return TimeSliderController; }
+	UWorld* GetWorld(UObject* OptionalObject = nullptr);
+	class AActor* GetVisualLoggerHelperActor();
+	FVisualLoggerEvents& GetVisualLoggerEvents() { return VisualLoggerEvents; }
 
-	const FVisualLoggerEvents& GetVisualLoggerEvents() { return VisualLoggerEvents; }
-	class FSequencerTimeSliderController* GetTimeSliderController() { return TimeSliderController.Get(); }
-	void SetTimeSliderController(TSharedPtr<class FSequencerTimeSliderController> InTimeSliderController) { TimeSliderController = InTimeSliderController; }
-	FFiltersPreset& GetFiltersPreset() { return FiltersPreset; }
+	TSharedPtr<IVisualLoggerInterface> GetVisualLoggerInterface() { return VisualLoggerInterface.Pin(); }
+	void SetVisualLoggerInterface(TSharedPtr<IVisualLoggerInterface> InVisualLoggerInterface) { VisualLoggerInterface = InVisualLoggerInterface; }
+	void SetCurrentVisualizer(TSharedPtr<class SVisualLogger> Visualizer) { CurrentVisualizer = Visualizer; }
 
+	void OnObjectSelectionChanged(TSharedPtr<class STimeline> TimeLine) { CurrentTimeLine = TimeLine; }
+
+	/** Static access */
+	static void Initialize();
+	static void Shutdown();
+	static FLogVisualizer& Get();
 protected:
-	FFiltersPreset	FiltersPreset;
+	static TSharedPtr< struct FLogVisualizer > StaticInstance;
+	
+	TSharedPtr<FSequencerTimeSliderController> TimeSliderController;
 	FVisualLoggerEvents VisualLoggerEvents;
-	TSharedPtr<class FSequencerTimeSliderController> TimeSliderController;
+	TWeakPtr<IVisualLoggerInterface> VisualLoggerInterface;
+	TWeakPtr<class STimeline> CurrentTimeLine;
+	TWeakPtr<class SVisualLogger> CurrentVisualizer;
 };
 
-
 #include "LogVisualizerStyle.h"
 #include "SVisualLogger.h"
 #include "SVisualLoggerToolbar.h"
diff --git a/Engine/Source/Developer/LogVisualizer/Private/LogVisualizerModule.cpp b/Engine/Source/Developer/LogVisualizer/Private/LogVisualizerModule.cpp
--- a/Engine/Source/Developer/LogVisualizer/Private/LogVisualizerModule.cpp
+++ b/Engine/Source/Developer/LogVisualizer/Private/LogVisualizerModule.cpp
@@ -1,7 +1,7 @@
 // Copyright 1998-2015 Epic Games, Inc. All Rights Reserved.
 
 #include "LogVisualizer.h"
-#include "LogVisualizerModule.h"
+#include "Runtime/Core/Public/Features/IModularFeatures.h"
 #include "LogVisualizerStyle.h"
 #include "SDockTab.h"
 #include "VisualLoggerRenderingActor.h"
@@ -19,9 +19,26 @@ static const FName VisualLoggerTabName("VisualLogger");
 
 //DEFINE_LOG_CATEGORY(LogLogVisualizer);
 
+class FLogVisualizerModule : public ILogVisualizer
+{
+public:
+	// Begin IModuleInterface
+	virtual void StartupModule() override;
+	virtual void ShutdownModule() override;
+	// End IModuleInterface
+
+	virtual void Goto(float Timestamp, FName LogOwner) override;
+	virtual void GotoNextItem() override;
+	virtual void GotoPreviousItem() override;
+
+private:
+	TSharedRef<SDockTab> SpawnLogVisualizerTab(const FSpawnTabArgs& SpawnTabArgs);
+};
+
 void FLogVisualizerModule::StartupModule()
 {
 	FLogVisualizerStyle::Initialize();
+	FLogVisualizer::Initialize();
 
 	FVisualLoggerCommands::Register();
 	IModularFeatures::Get().RegisterModularFeature(VisualLoggerTabName, this);
@@ -57,13 +74,14 @@ void FLogVisualizerModule::ShutdownModule()
 		SettingsModule->UnregisterSettings("Editor", "General", "VisualLogger");
 	}
 
+	FLogVisualizer::Shutdown();
 	FLogVisualizerStyle::Shutdown();
 }
 
 TSharedRef<SDockTab> FLogVisualizerModule::SpawnLogVisualizerTab(const FSpawnTabArgs& SpawnTabArgs)
 {
 	const TSharedRef<SDockTab> MajorTab = SNew(SDockTab)
-		.TabRole(ETabRole::MajorTab);
+		.TabRole(ETabRole::NomadTab);
 
 	TSharedPtr<SWidget> TabContent;
 
@@ -74,5 +92,21 @@ TSharedRef<SDockTab> FLogVisualizerModule::SpawnLogVisualizerTab(const FSpawnTab
 	return MajorTab;
 }
 
+void FLogVisualizerModule::Goto(float Timestamp, FName LogOwner)
+{
+	FLogVisualizer::Get().Goto(Timestamp, LogOwner);
+}
+
+void FLogVisualizerModule::GotoNextItem()
+{
+	FLogVisualizer::Get().GotoNextItem();
+}
+
+void FLogVisualizerModule::GotoPreviousItem()
+{
+	FLogVisualizer::Get().GotoPreviousItem();
+}
+
+
 IMPLEMENT_MODULE(FLogVisualizerModule, LogVisualizer);
 #undef LOCTEXT_NAMESPACE
diff --git a/Engine/Source/Developer/LogVisualizer/Private/LogVisualizerStyle.cpp b/Engine/Source/Developer/LogVisualizer/Private/LogVisualizerStyle.cpp
--- a/Engine/Source/Developer/LogVisualizer/Private/LogVisualizerStyle.cpp
+++ b/Engine/Source/Developer/LogVisualizer/Private/LogVisualizerStyle.cpp
@@ -1,6 +1,5 @@
 // Copyright 1998-2015 Epic Games, Inc. All Rights Reserved.
 #include "LogVisualizer.h"
-#include "LogVisualizerModule.h"
 #include "LogVisualizerStyle.h"
 
 #define IMAGE_BRUSH( RelativePath, ... ) FSlateImageBrush( FPaths::EngineContentDir() / "Editor/Slate"/ RelativePath + TEXT(".png"), __VA_ARGS__ )
diff --git a/Engine/Source/Developer/LogVisualizer/Private/STimeline.cpp b/Engine/Source/Developer/LogVisualizer/Private/STimeline.cpp
--- a/Engine/Source/Developer/LogVisualizer/Private/STimeline.cpp
+++ b/Engine/Source/Developer/LogVisualizer/Private/STimeline.cpp
@@ -130,7 +130,7 @@ void STimeline::UpdateVisibilityForEntry(FVisualLogDevice::FVisualLogEntryItem&
 {
 	TArray<FVisualLoggerCategoryVerbosityPair> OutCategories;
 	FVisualLoggerHelpers::GetCategories(CurrentEntry.Entry, OutCategories);
-	bool bHasValidCategories = VisualLoggerInterface.Pin()->HasValidCategories(OutCategories);
+	bool bHasValidCategories = FLogVisualizer::Get().GetVisualLoggerInterface()->HasValidCategories(OutCategories);
 
 	if (bSearchInsideLogs && bHasValidCategories && SearchString.Len() > 0)
 	{
@@ -185,7 +185,6 @@ void STimeline::AddEntry(const FVisualLogDevice::FVisualLogEntryItem& Entry)
 
 void STimeline::Construct(const FArguments& InArgs, TSharedPtr<SVisualLoggerView> VisualLoggerView, TSharedPtr<FSequencerTimeSliderController> TimeSliderController, TSharedPtr<STimelinesContainer> InContainer, const FVisualLogDevice::FVisualLogEntryItem& Entry)
 {
-	VisualLoggerInterface = InArgs._VisualLoggerInterface.Get();
 	OnGetMenuContent = InArgs._OnGetMenuContent;
 
 	Owner = InContainer;
@@ -251,7 +250,6 @@ void STimeline::Construct(const FArguments& InArgs, TSharedPtr<SVisualLoggerView
 				[
 					// Search box for searching through the outliner
 					SAssignNew(TimelineBar, STimelineBar, TimeSliderController, SharedThis(this))
-					.VisualLoggerInterface(InArgs._VisualLoggerInterface)
 				]
 			]
 		];
@@ -270,3 +268,18 @@ const FSlateBrush* STimeline::GetBorder() const
 		return FLogVisualizerStyle::Get().GetBrush("ToolBar.Button.Normal");
 	}
 }
+
+void STimeline::Goto(float ScrubPosition) 
+{ 
+	TimelineBar->SnapScrubPosition(ScrubPosition); 
+}
+
+void STimeline::GotoNextItem() 
+{ 
+	TimelineBar->GotoNextItem(); 
+}
+
+void STimeline::GotoPreviousItem() 
+{ 
+	TimelineBar->GotoPreviousItem(); 
+}
diff --git a/Engine/Source/Developer/LogVisualizer/Private/STimeline.h b/Engine/Source/Developer/LogVisualizer/Private/STimeline.h
--- a/Engine/Source/Developer/LogVisualizer/Private/STimeline.h
+++ b/Engine/Source/Developer/LogVisualizer/Private/STimeline.h
@@ -9,7 +9,6 @@ class STimeline : public SCompoundWidget
 {
 public:
 	SLATE_BEGIN_ARGS(STimeline){}
-		SLATE_ATTRIBUTE(TWeakPtr<IVisualLoggerInterface>, VisualLoggerInterface)
 		SLATE_EVENT(FOnItemSelectionChanged, OnItemSelectionChanged)
 		SLATE_EVENT(FOnGetContent, OnGetMenuContent)
 	SLATE_END_ARGS();
@@ -39,13 +38,17 @@ public:
 	FName GetOwnerClassName() { return OwnerClassName; }
 	void HandleLogVisualizerSettingChanged(FName Name);
 	void UpdateVisibilityForEntry(FVisualLogDevice::FVisualLogEntryItem& Entry, const FString& SearchString, bool SearchInsideLogs);
+	TSharedPtr<class STimelinesContainer> GetOwner() { return Owner; }
+
+	void Goto(float ScrubPosition);
+	void GotoNextItem();
+	void GotoPreviousItem();
 
 protected:
 	TArray<FVisualLogDevice::FVisualLogEntryItem> Entries;
 	TArray<const FVisualLogDevice::FVisualLogEntryItem*> HiddenEntries;
 	TSharedPtr<class STimelinesContainer> Owner;
 	TSharedPtr<class STimelineBar> TimelineBar;
-	TWeakPtr<IVisualLoggerInterface> VisualLoggerInterface;
 	TSharedPtr<SMenuAnchor> PopupAnchor;
 
 	FName Name;
diff --git a/Engine/Source/Developer/LogVisualizer/Private/STimelineBar.cpp b/Engine/Source/Developer/LogVisualizer/Private/STimelineBar.cpp
--- a/Engine/Source/Developer/LogVisualizer/Private/STimelineBar.cpp
+++ b/Engine/Source/Developer/LogVisualizer/Private/STimelineBar.cpp
@@ -43,7 +43,7 @@ FReply STimelineBar::OnMouseButtonDoubleClick(const FGeometry& MyGeometry, const
 {
 	auto &Entries = TimelineOwner.Pin()->GetEntries();
 	SnapScrubPosition(TimeSliderController->GetTimeSliderArgs().ScrubPosition.Get());
-	UWorld* World = VisualLoggerInterface.Pin()->GetWorld();
+	UWorld* World = FLogVisualizer::Get().GetWorld();
 	if (World && MouseEvent.GetEffectingButton() == EKeys::LeftMouseButton && Entries.IsValidIndex(CurrentItemIndex))
 	{
 		FVector CurrentLocation = Entries[CurrentItemIndex].Entry.Location;
@@ -109,6 +109,66 @@ FReply STimelineBar::OnMouseButtonDoubleClick(const FGeometry& MyGeometry, const
 	return FReply::Unhandled();
 }
 
+void STimelineBar::GotoNextItem(int32 MoveDistance)
+{
+	int32 NewItemIndex = CurrentItemIndex;
+	auto &Entries = TimelineOwner.Pin()->GetEntries();
+	int32 Index = 0;
+	while (true)
+	{
+		NewItemIndex++;
+		if (Entries.IsValidIndex(NewItemIndex))
+		{
+			auto& CurrentEntryItem = Entries[NewItemIndex];
+			if (TimelineOwner.Pin()->IsEntryHidden(CurrentEntryItem) == false && ++Index == MoveDistance)
+			{
+				break;
+			}
+		}
+		else
+		{
+			NewItemIndex = FMath::Clamp(NewItemIndex, 0, Entries.Num() - 1);
+			break;
+		}
+	}
+
+	if (NewItemIndex != CurrentItemIndex)
+	{
+		float NewTimeStamp = Entries[NewItemIndex].Entry.TimeStamp;
+		SnapScrubPosition(NewItemIndex);
+	}
+}
+
+void STimelineBar::GotoPreviousItem(int32 MoveDistance)
+{
+	int32 NewItemIndex = CurrentItemIndex;
+	auto &Entries = TimelineOwner.Pin()->GetEntries();
+	int32 Index = 0;
+	while (true)
+	{
+		NewItemIndex--;
+		if (Entries.IsValidIndex(NewItemIndex))
+		{
+			auto& CurrentEntryItem = Entries[NewItemIndex];
+			if (TimelineOwner.Pin()->IsEntryHidden(CurrentEntryItem) == false && ++Index == MoveDistance)
+			{
+				break;
+			}
+		}
+		else
+		{
+			NewItemIndex = FMath::Clamp(NewItemIndex, 0, Entries.Num() - 1);
+			break;
+		}
+	}
+
+	if (NewItemIndex != CurrentItemIndex)
+	{
+		float NewTimeStamp = Entries[NewItemIndex].Entry.TimeStamp;
+		SnapScrubPosition(NewItemIndex);
+	}
+}
+
 FReply STimelineBar::OnKeyDown(const FGeometry& MyGeometry, const FKeyEvent& InKeyEvent)
 {
 	FReply ReturnValue = FReply::Unhandled();
@@ -135,6 +195,10 @@ FReply STimelineBar::OnKeyDown(const FGeometry& MyGeometry, const FKeyEvent& InK
 					break;
 				}
 			}
+			if (NewItemIndex != CurrentItemIndex)
+			{
+				SnapScrubPosition(NewItemIndex);
+			}
 			ReturnValue = FReply::Handled();
 		}
 		else if (Key == EKeys::End)
@@ -147,38 +211,21 @@ FReply STimelineBar::OnKeyDown(const FGeometry& MyGeometry, const FKeyEvent& InK
 					break;
 				}
 			}
+			if (NewItemIndex != CurrentItemIndex)
+			{
+				SnapScrubPosition(NewItemIndex);
+			}
 			ReturnValue = FReply::Handled();
 		}
-		else if (Key == EKeys::Left || Key == EKeys::Right)
+		else if (Key == EKeys::Left)
 		{
-			int32 Direction = Key == EKeys::Left ? -1 : 1;
-			int32 MoveBy = InKeyEvent.IsLeftControlDown() ? InKeyEvent.IsLeftShiftDown() ? 20 : 10 : 1;
-
-			int32 Index = 0;
-			while (true)
-			{
-				NewItemIndex += Direction;
-				if (Entries.IsValidIndex(NewItemIndex))
-				{
-					auto& CurrentEntryItem = Entries[NewItemIndex];
-					if (TimelineOwner.Pin()->IsEntryHidden(CurrentEntryItem) == false && ++Index == MoveBy)
-					{
-						break;
-					}
-				}
-				else
-				{
-					NewItemIndex = FMath::Clamp(NewItemIndex, 0, Entries.Num() - 1);
-					break;
-				}
-			}
+			GotoPreviousItem(InKeyEvent.IsLeftControlDown() ? InKeyEvent.IsLeftShiftDown() ? 20 : 10 : 1);
 			ReturnValue = FReply::Handled();
 		}
-
-		if (NewItemIndex != CurrentItemIndex)
+		else if (Key == EKeys::Right)
 		{
-			float NewTimeStamp = Entries[NewItemIndex].Entry.TimeStamp;
-			SnapScrubPosition(NewItemIndex);
+			GotoNextItem(InKeyEvent.IsLeftControlDown() ? InKeyEvent.IsLeftShiftDown() ? 20 : 10 : 1);
+			ReturnValue = FReply::Handled();
 		}
 	}
 
@@ -225,7 +272,7 @@ void STimelineBar::SnapScrubPosition(float ScrubPosition)
 		if (CurrentItemIndex != BestItemIndex)
 		{
 			CurrentItemIndex = BestItemIndex;
-			VisualLoggerInterface.Pin()->GetVisualLoggerEvents().OnItemSelectionChanged.ExecuteIfBound(Entries[BestItemIndex]);
+			FLogVisualizer::Get().GetVisualLoggerEvents().OnItemSelectionChanged.ExecuteIfBound(Entries[BestItemIndex]);
 		}
 		TimeSliderController->CommitScrubPosition(CurrentTime, false);
 	}
@@ -240,7 +287,7 @@ void STimelineBar::SnapScrubPosition(int32 NewItemIndex)
 		if (CurrentItemIndex != NewItemIndex)
 		{
 			CurrentItemIndex = NewItemIndex;
-			VisualLoggerInterface.Pin()->GetVisualLoggerEvents().OnItemSelectionChanged.ExecuteIfBound(Entries[NewItemIndex]);
+			FLogVisualizer::Get().GetVisualLoggerEvents().OnItemSelectionChanged.ExecuteIfBound(Entries[NewItemIndex]);
 		}
 		TimeSliderController->CommitScrubPosition(CurrentTime, false);
 	}
@@ -248,8 +295,6 @@ void STimelineBar::SnapScrubPosition(int32 NewItemIndex)
 
 void STimelineBar::Construct(const FArguments& InArgs, TSharedPtr<FSequencerTimeSliderController> InTimeSliderController, TSharedPtr<STimeline> InTimelineOwner)
 {
-	VisualLoggerInterface = InArgs._VisualLoggerInterface.Get();
-
 	TimeSliderController = InTimeSliderController;
 	TimelineOwner = InTimelineOwner;
 	CurrentItemIndex = INDEX_NONE;
@@ -264,7 +309,7 @@ FVector2D STimelineBar::ComputeDesiredSize(float) const
 
 void STimelineBar::OnSelect()
 {
-	FSlateApplication::Get().SetKeyboardFocus(SharedThis(this), EFocusCause::Navigation);
+	FSlateApplication::Get().SetKeyboardFocus(SharedThis(this), EFocusCause::SetDirectly);
 	CurrentItemIndex = INDEX_NONE;
 }
 
@@ -431,7 +476,7 @@ int32 STimelineBar::OnPaint(const FPaintArgs& Args, const FGeometry& AllottedGeo
 		auto &Entries = TimelineOwner.Pin()->GetEntries();
 		if (BestItemIndex != CurrentItemIndex)
 		{
-			VisualLoggerInterface.Pin()->GetVisualLoggerEvents().OnItemSelectionChanged.ExecuteIfBound(Entries[BestItemIndex]);
+			FLogVisualizer::Get().GetVisualLoggerEvents().OnItemSelectionChanged.ExecuteIfBound(Entries[BestItemIndex]);
 		}
 
 		float CurrentTime = Entries[BestItemIndex].Entry.TimeStamp;
diff --git a/Engine/Source/Developer/LogVisualizer/Private/STimelineBar.h b/Engine/Source/Developer/LogVisualizer/Private/STimelineBar.h
--- a/Engine/Source/Developer/LogVisualizer/Private/STimelineBar.h
+++ b/Engine/Source/Developer/LogVisualizer/Private/STimelineBar.h
@@ -11,7 +11,6 @@ class STimelineBar : public SLeafWidget
 public:
 	SLATE_BEGIN_ARGS(STimelineBar){}
 		SLATE_EVENT(FOnItemSelectionChanged, OnItemSelectionChanged)
-		SLATE_ATTRIBUTE(TWeakPtr<IVisualLoggerInterface>, VisualLoggerInterface)
 	SLATE_END_ARGS()
 
 	/** Constructs this widget with InArgs */
@@ -32,9 +31,10 @@ public:
 
 	void OnSelect();
 	void OnDeselect();
+	void GotoNextItem(int32 MoveDistance = 1);
+	void GotoPreviousItem(int32 MoveDistance = 1);
 
 protected:
-	TWeakPtr<IVisualLoggerInterface> VisualLoggerInterface;
 	TSharedPtr<class FSequencerTimeSliderController> TimeSliderController;
 	TWeakPtr<class STimeline> TimelineOwner;
 	mutable int32 CurrentItemIndex;
diff --git a/Engine/Source/Developer/LogVisualizer/Private/STimelinesContainer.cpp b/Engine/Source/Developer/LogVisualizer/Private/STimelinesContainer.cpp
--- a/Engine/Source/Developer/LogVisualizer/Private/STimelinesContainer.cpp
+++ b/Engine/Source/Developer/LogVisualizer/Private/STimelinesContainer.cpp
@@ -16,8 +16,7 @@ TSharedRef<SWidget> STimelinesContainer::MakeTimeline(TSharedPtr<class SVisualLo
 	ContainingBorder->AddSlot()
 		[
 			SAssignNew(NewTimeline, STimeline, VisualLoggerView, TimeSliderController, SharedThis(this), Entry)
-			.OnItemSelectionChanged(this->VisualLoggerInterface.Pin()->GetVisualLoggerEvents().OnItemSelectionChanged)
-			.VisualLoggerInterface(this->VisualLoggerInterface)
+			.OnItemSelectionChanged(FLogVisualizer::Get().GetVisualLoggerEvents().OnItemSelectionChanged)
 			.OnGetMenuContent(this, &STimelinesContainer::GetRightClickMenuContent)
 		];
 
@@ -72,7 +71,7 @@ void STimelinesContainer::SetSelectionState(TSharedPtr<class STimeline> Affected
 			SelectedNodes.Add(AffectedNode);
 			AffectedNode->OnSelect();
 		}
-		VisualLoggerInterface.Pin()->GetVisualLoggerEvents().OnObjectSelectionChanged.ExecuteIfBound(AffectedNode);
+		FLogVisualizer::Get().GetVisualLoggerEvents().OnObjectSelectionChanged.ExecuteIfBound(AffectedNode);
 	}
 	else if (AffectedNode.IsValid())
 	{
@@ -272,7 +271,6 @@ void STimelinesContainer::Construct(const FArguments& InArgs, TSharedRef<class S
 {
 	TimeSliderController = InTimeSliderController;
 	VisualLoggerView = InVisualLoggerView;
-	VisualLoggerInterface = InArgs._VisualLoggerInterface.Get();
 
 	ChildSlot
 		[
diff --git a/Engine/Source/Developer/LogVisualizer/Private/STimelinesContainer.h b/Engine/Source/Developer/LogVisualizer/Private/STimelinesContainer.h
--- a/Engine/Source/Developer/LogVisualizer/Private/STimelinesContainer.h
+++ b/Engine/Source/Developer/LogVisualizer/Private/STimelinesContainer.h
@@ -9,7 +9,6 @@ class STimelinesContainer : public SCompoundWidget
 {
 public:
 	SLATE_BEGIN_ARGS(STimelinesContainer){}
-		SLATE_ATTRIBUTE(TWeakPtr<IVisualLoggerInterface>, VisualLoggerInterface)
 	SLATE_END_ARGS()
 
 	virtual FReply OnMouseButtonDown(const FGeometry& MyGeometry, const FPointerEvent& MouseEvent) override;
@@ -66,6 +65,5 @@ protected:
 	TArray<TSharedPtr<class STimeline> > TimelineItems;
 	TArray< TSharedPtr<class STimeline> > SelectedNodes;
 	TSharedPtr<SVerticalBox> ContainingBorder;
-	TWeakPtr<IVisualLoggerInterface> VisualLoggerInterface;
 	float CachedMinTime, CachedMaxTime;
 };
diff --git a/Engine/Source/Developer/LogVisualizer/Private/SVisualLogger.cpp b/Engine/Source/Developer/LogVisualizer/Private/SVisualLogger.cpp
--- a/Engine/Source/Developer/LogVisualizer/Private/SVisualLogger.cpp
+++ b/Engine/Source/Developer/LogVisualizer/Private/SVisualLogger.cpp
@@ -40,10 +40,9 @@ namespace LogVisualizer
 
 struct FVisualLoggerInterface : public IVisualLoggerInterface
 {
-	FVisualLoggerInterface(TSharedPtr<SVisualLogger> InVisualLogger, FVisualLoggerEvents InVisualLoggerEvents)
+	FVisualLoggerInterface(TSharedPtr<SVisualLogger> InVisualLogger)
 		: VisualLogger(InVisualLogger)
 	{
-		VisualLoggerEvents = InVisualLoggerEvents;
 	}
 
 	virtual ~FVisualLoggerInterface() 
@@ -73,51 +72,6 @@ struct FVisualLoggerInterface : public IVisualLoggerInterface
 		return VisualLogger.Pin()->GetVisualLoggerFilters()->IsFilterEnabled(InGraphName, InDataName, Verbosity);
 	}
 
-	virtual FLinearColor GetCategoryColor(FName Category) override
-	{
-		return VisualLogger.Pin()->GetVisualLoggerFilters()->GetColorForUsedCategory(Category.ToString());
-	}
-
-	UWorld* GetWorld() const
-	{
-		UWorld* World = NULL;
-		UEditorEngine *EEngine = Cast<UEditorEngine>(GEngine);
-		if (GIsEditor && EEngine != NULL)
-		{
-			// lets use PlayWorld during PIE/Simulate and regular world from editor otherwise, to draw debug information
-			World = EEngine->PlayWorld != NULL ? EEngine->PlayWorld : EEngine->GetEditorWorldContext().World();
-
-		}
-		else if (!GIsEditor)
-		{
-
-			World = GEngine->GetWorld();
-		}
-
-		if (World == NULL)
-		{
-			World = GWorld;
-		}
-
-		return World;
-	}
-
-	class AActor* GetVisualLoggerHelperActor() override
-	{
-		UWorld* World = GetWorld();
-		for (TActorIterator<AVisualLoggerRenderingActor> It(World); It; ++It)
-		{
-			return *It;
-		}
-
-		FActorSpawnParameters SpawnInfo;
-		SpawnInfo.bNoCollisionFail = true;
-		SpawnInfo.Name = *FString::Printf(TEXT("VisualLoggerRenderingActor"));
-		AActor* HelperActor = World->SpawnActor<AVisualLoggerRenderingActor>(SpawnInfo);
-
-		return HelperActor;
-	}
-
 protected:
 	TWeakPtr<SVisualLogger> VisualLogger;
 };
@@ -209,26 +163,22 @@ SVisualLogger::~SVisualLogger()
 	VisualLoggerInterface.Reset();
 }
 
-void SVisualLogger::OnTabLosed()
-{
-
-}
-
 /* SMessagingDebugger interface
 *****************************************************************************/
 
 void SVisualLogger::Construct(const FArguments& InArgs, const TSharedRef<SDockTab>& ConstructUnderMajorTab, const TSharedPtr<SWindow>& ConstructUnderWindow)
 {
+	FLogVisualizer::Get().SetCurrentVisualizer(SharedThis(this));
 	//////////////////////////////////////////////////////////////////////////
 	// Visual Logger Events
-	FVisualLoggerEvents VisualLoggerEvents;
-	VisualLoggerEvents.OnItemSelectionChanged = FOnItemSelectionChanged::CreateRaw(this, &SVisualLogger::OnItemSelectionChanged);
-	VisualLoggerEvents.OnFiltersChanged = FOnFiltersChanged::CreateRaw(this, &SVisualLogger::OnFiltersChanged);
-	VisualLoggerEvents.OnObjectSelectionChanged = FOnObjectSelectionChanged::CreateRaw(this, &SVisualLogger::OnObjectSelectionChanged);
+	FLogVisualizer::Get().GetVisualLoggerEvents().OnItemSelectionChanged = FOnItemSelectionChanged::CreateRaw(this, &SVisualLogger::OnItemSelectionChanged);
+	FLogVisualizer::Get().GetVisualLoggerEvents().OnFiltersChanged = FOnFiltersChanged::CreateRaw(this, &SVisualLogger::OnFiltersChanged);
+	FLogVisualizer::Get().GetVisualLoggerEvents().OnObjectSelectionChanged = FOnObjectSelectionChanged::CreateRaw(this, &SVisualLogger::OnObjectSelectionChanged);
 
 	//////////////////////////////////////////////////////////////////////////
 	// Visual Logger Interface
-	VisualLoggerInterface = MakeShareable(new FVisualLoggerInterface(SharedThis(this), VisualLoggerEvents));
+	VisualLoggerInterface = MakeShareable(new FVisualLoggerInterface(SharedThis(this)));
+	FLogVisualizer::Get().SetVisualLoggerInterface(VisualLoggerInterface);
 
 	//////////////////////////////////////////////////////////////////////////
 	// Command Action Lists
@@ -260,6 +210,13 @@ void SVisualLogger::Construct(const FArguments& InArgs, const TSharedRef<SDockTa
 		FIsActionChecked(), 
 		FIsActionButtonVisible::CreateRaw(this, &SVisualLogger::HandleSaveCommandCanExecute));
 
+	ActionList.MapAction(
+		Commands.MoveCursorLeft,
+		FExecuteAction::CreateRaw(this, &SVisualLogger::OnMoveCursorLeftCommand));
+	ActionList.MapAction(
+		Commands.MoveCursorRight,
+		FExecuteAction::CreateRaw(this, &SVisualLogger::OnMoveCursorRightCommand));
+
 
 	// Tab Spawners
 	TabManager = FGlobalTabmanager::Get()->NewTabManager(ConstructUnderMajorTab);
@@ -392,12 +349,18 @@ void SVisualLogger::Construct(const FArguments& InArgs, const TSharedRef<SDockTa
 			]
 		];
 
-	VisualLoggerCanvasRenderer = MakeShareable(new FVisualLoggerCanvasRenderer(VisualLoggerInterface));
+	VisualLoggerCanvasRenderer = MakeShareable(new FVisualLoggerCanvasRenderer());
 
 	DrawOnCanvasDelegateHandle = UDebugDrawService::Register(TEXT("VisLog"), FDebugDrawDelegate::CreateRaw(VisualLoggerCanvasRenderer.Get(), &FVisualLoggerCanvasRenderer::DrawOnCanvas));
 	//UGameplayDebuggingComponent::OnDebuggingTargetChangedDelegate.AddSP(this, &SVisualLogger::SelectionChanged);
 }
 
+FReply SVisualLogger::OnKeyDown(const FGeometry& MyGeometry, const FKeyEvent& InKeyEvent)
+{
+	FUICommandList& ActionList = *CommandList;
+	return ActionList.ProcessCommandBindings(InKeyEvent) ? FReply::Handled() : FReply::Unhandled();
+}
+
 void SVisualLogger::HandleMajorTabPersistVisualState()
 {
 	// save any settings here
@@ -477,17 +440,17 @@ TSharedRef<SDockTab> SVisualLogger::HandleTabManagerSpawnTab(const FSpawnTabArgs
 	}
 	else if (TabIdentifier == FiltersTabId)
 	{
-		TabWidget = SAssignNew(VisualLoggerFilters, SVisualLoggerFilters, CommandList, VisualLoggerInterface);
+		TabWidget = SAssignNew(VisualLoggerFilters, SVisualLoggerFilters, CommandList);
 		AutoSizeTab = true;
 	}
 	else if (TabIdentifier == MainViewTabId)
 	{
-		TabWidget = SAssignNew(MainView, SVisualLoggerView, CommandList, VisualLoggerInterface).OnFiltersSearchChanged(this, &SVisualLogger::OnFiltersSearchChanged);
+		TabWidget = SAssignNew(MainView, SVisualLoggerView, CommandList).OnFiltersSearchChanged(this, &SVisualLogger::OnFiltersSearchChanged);
 		AutoSizeTab = false;
 	}
 	else if (TabIdentifier == LogsListTabId)
 	{
-		TabWidget = SAssignNew(LogsList, SVisualLoggerLogsList, CommandList, VisualLoggerInterface);
+		TabWidget = SAssignNew(LogsList, SVisualLoggerLogsList, CommandList);
 		AutoSizeTab = false;
 	}
 	else if (TabIdentifier == StatusViewTabId)
@@ -531,7 +494,7 @@ void SVisualLogger::HandleStopRecordingCommandExecute()
 {
 	FVisualLogger::Get().SetIsRecording(false);
 
-	UWorld* World = VisualLoggerInterface->GetWorld();
+	UWorld* World = FLogVisualizer::Get().GetWorld();
 	if (AVisualLoggerCameraController::IsEnabled(World))
 	{
 		AVisualLoggerCameraController::DisableCamera(World);
@@ -546,13 +509,13 @@ bool SVisualLogger::HandleStopRecordingCommandIsVisible() const
 
 bool SVisualLogger::HandlePauseCommandCanExecute() const
 {
-	UWorld* World = VisualLoggerInterface->GetWorld();
+	UWorld* World = FLogVisualizer::Get().GetWorld();
 	return FVisualLogger::Get().IsRecording() && World && !World->bPlayersOnly && !World->bPlayersOnlyPending;
 }
 
 void SVisualLogger::HandlePauseCommandExecute()
 {
-	UWorld* World = VisualLoggerInterface->GetWorld();
+	UWorld* World = FLogVisualizer::Get().GetWorld();
 	if (World != NULL)
 	{
 		World->bPlayersOnlyPending = true;
@@ -566,13 +529,13 @@ bool SVisualLogger::HandlePauseCommandIsVisible() const
 
 bool SVisualLogger::HandleResumeCommandCanExecute() const
 {
-	UWorld* World = VisualLoggerInterface->GetWorld();
+	UWorld* World = FLogVisualizer::Get().GetWorld();
 	return FVisualLogger::Get().IsRecording() && World && (World->bPlayersOnly || World->bPlayersOnlyPending);
 }
 
 void SVisualLogger::HandleResumeCommandExecute()
 {
-	UWorld* World = VisualLoggerInterface->GetWorld();
+	UWorld* World = FLogVisualizer::Get().GetWorld();
 	if (World != NULL)
 	{
 		World->bPlayersOnly = false;
@@ -587,19 +550,19 @@ bool SVisualLogger::HandleResumeCommandIsVisible() const
 
 bool SVisualLogger::HandleCameraCommandIsChecked() const
 {
-	UWorld* World = VisualLoggerInterface->GetWorld();
+	UWorld* World = FLogVisualizer::Get().GetWorld();
 	return World && AVisualLoggerCameraController::IsEnabled(World);
 }
 
 bool SVisualLogger::HandleCameraCommandCanExecute() const
 {
-	UWorld* World = VisualLoggerInterface->GetWorld();
+	UWorld* World = FLogVisualizer::Get().GetWorld();
 	return FVisualLogger::Get().IsRecording() && World && (World->bPlayersOnly || World->bPlayersOnlyPending) && World->IsPlayInEditor() && (GEditor && !GEditor->bIsSimulatingInEditor);
 }
 
 void SVisualLogger::HandleCameraCommandExecute()
 {
-	UWorld* World = VisualLoggerInterface->GetWorld();
+	UWorld* World = FLogVisualizer::Get().GetWorld();
 	if (AVisualLoggerCameraController::IsEnabled(World))
 	{
 		AVisualLoggerCameraController::DisableCamera(World);
@@ -799,10 +762,10 @@ void SVisualLogger::OnItemSelectionChanged(const FVisualLogDevice::FVisualLogEnt
 	StatusView->OnItemSelectionChanged(EntryItem);
 	VisualLoggerCanvasRenderer->OnItemSelectionChanged(EntryItem.Entry);
 	VisualLoggerFilters->OnItemSelectionChanged(EntryItem.Entry);
-	AVisualLoggerRenderingActor* HelperActor = Cast<AVisualLoggerRenderingActor>(VisualLoggerInterface->GetVisualLoggerHelperActor());
+	AVisualLoggerRenderingActor* HelperActor = Cast<AVisualLoggerRenderingActor>(FLogVisualizer::Get().GetVisualLoggerHelperActor());
 	if (HelperActor)
 	{
-		HelperActor->OnItemSelectionChanged(EntryItem, VisualLoggerInterface);
+		HelperActor->OnItemSelectionChanged(EntryItem);
 	}
 }
 
@@ -815,12 +778,13 @@ void SVisualLogger::OnFiltersChanged()
 void SVisualLogger::OnObjectSelectionChanged(TSharedPtr<class STimeline> TimeLine)
 {
 	VisualLoggerCanvasRenderer->ObjectSelectionChanged(TimeLine);
-	AVisualLoggerRenderingActor* HelperActor = Cast<AVisualLoggerRenderingActor>(VisualLoggerInterface->GetVisualLoggerHelperActor());
+	AVisualLoggerRenderingActor* HelperActor = Cast<AVisualLoggerRenderingActor>(FLogVisualizer::Get().GetVisualLoggerHelperActor());
 	if (HelperActor)
 	{
 		HelperActor->ObjectSelectionChanged(TimeLine);
 	}
 	MainView->OnObjectSelectionChanged(TimeLine);
+	FLogVisualizer::Get().OnObjectSelectionChanged(TimeLine);
 }
 
 void SVisualLogger::OnFiltersSearchChanged(const FText& Filter)
@@ -841,4 +805,23 @@ void SVisualLogger::OnFiltersSearchChanged(const FText& Filter)
 		MainView->OnFiltersSearchChanged(Filter);
 	}
 }
+
+void SVisualLogger::OnMoveCursorLeftCommand()
+{
+	FLogVisualizer::Get().GotoNextItem();
+}
+
+void SVisualLogger::OnMoveCursorRightCommand()
+{
+	FLogVisualizer::Get().GotoPreviousItem();
+}
+
+void SVisualLogger::GetTimelines(TArray<TSharedPtr<class STimeline> >& TimeLines, bool bOnlySelectedOnes)
+{
+	if (MainView.IsValid())
+	{
+		MainView->GetTimelines(TimeLines, bOnlySelectedOnes);
+	}
+}
+
 #undef LOCTEXT_NAMESPACE
diff --git a/Engine/Source/Developer/LogVisualizer/Private/SVisualLogger.h b/Engine/Source/Developer/LogVisualizer/Private/SVisualLogger.h
--- a/Engine/Source/Developer/LogVisualizer/Private/SVisualLogger.h
+++ b/Engine/Source/Developer/LogVisualizer/Private/SVisualLogger.h
@@ -56,9 +56,11 @@ public:
 	static void FillWindowMenu(FMenuBuilder& MenuBuilder, const TSharedPtr<FTabManager> TabManager);
 	void FillLoadPresetMenu(FMenuBuilder& Builder);
 
+	void GetTimelines(TArray<TSharedPtr<class STimeline> >&, bool bOnlySelectedOnes = false);
+
 	/** Callback for for when the owner tab's visual state is being persisted. */
 	void HandleMajorTabPersistVisualState();
-	void OnTabLosed();
+	virtual FReply OnKeyDown(const FGeometry& MyGeometry, const FKeyEvent& InKeyEvent) override;
 
 	void OnNewLogEntry(const FVisualLogDevice::FVisualLogEntryItem& Entry);
 	void CollectNewCategories(const FVisualLogDevice::FVisualLogEntryItem& Entry);
@@ -100,6 +102,10 @@ public:
 	void OnNewWorld(UWorld* NewWorld);
 	void ResetData();
 
+protected:
+	void OnMoveCursorLeftCommand();
+	void OnMoveCursorRightCommand();
+
 protected:
 	// Holds the list of UI commands.
 	TSharedRef<FUICommandList> CommandList;
diff --git a/Engine/Source/Developer/LogVisualizer/Private/SVisualLoggerFilters.cpp b/Engine/Source/Developer/LogVisualizer/Private/SVisualLoggerFilters.cpp
--- a/Engine/Source/Developer/LogVisualizer/Private/SVisualLoggerFilters.cpp
+++ b/Engine/Source/Developer/LogVisualizer/Private/SVisualLoggerFilters.cpp
@@ -12,59 +12,9 @@
 // SLogFilterList
 /////////////////////
 
-FColor SVisualLoggerFilters::ColorPalette[] = {
-	FColor(0xff00A480),
-	FColorList::Aquamarine,
-	FColorList::Cyan,
-	FColorList::Brown,
-	FColorList::Green,
-	FColorList::Orange,
-	FColorList::Magenta,
-	FColorList::BrightGold,
-	FColorList::NeonBlue,
-	FColorList::MediumSlateBlue,
-	FColorList::SpicyPink,
-	FColorList::SpringGreen,
-	FColorList::SteelBlue,
-	FColorList::SummerSky,
-	FColorList::Violet,
-	FColorList::VioletRed,
-	FColorList::YellowGreen,
-	FColor(0xff62E200),
-	FColor(0xff1F7B67),
-	FColor(0xff62AA2A),
-	FColor(0xff70227E),
-	FColor(0xff006B53),
-	FColor(0xff409300),
-	FColor(0xff5D016D),
-	FColor(0xff34D2AF),
-	FColor(0xff8BF13C),
-	FColor(0xffBC38D3),
-	FColor(0xff5ED2B8),
-	FColor(0xffA6F16C),
-	FColor(0xffC262D3),
-	FColor(0xff0F4FA8),
-	FColor(0xff00AE68),
-	FColor(0xffDC0055),
-	FColor(0xff284C7E),
-	FColor(0xff21825B),
-	FColor(0xffA52959),
-	FColor(0xff05316D),
-	FColor(0xff007143),
-	FColor(0xff8F0037),
-	FColor(0xff4380D3),
-	FColor(0xff36D695),
-	FColor(0xffEE3B80),
-	FColor(0xff6996D3),
-	FColor(0xff60D6A7),
-	FColor(0xffEE6B9E)
-};
-
 #define LOCTEXT_NAMESPACE "SVisualLoggerFilters"
-void SVisualLoggerFilters::Construct(const FArguments& InArgs, const TSharedRef<FUICommandList>& InCommandList, TSharedPtr<IVisualLoggerInterface> InVisualLoggerInterface)
+void SVisualLoggerFilters::Construct(const FArguments& InArgs, const TSharedRef<FUICommandList>& InCommandList)
 {
-	VisualLoggerInterface = InVisualLoggerInterface;
-
 	ChildSlot
 		[
 			SAssignNew(FilterBox, SWrapBox)
@@ -305,38 +255,6 @@ void SVisualLoggerFilters::InvalidateCanvas()
 #endif
 }
 
-FLinearColor SVisualLoggerFilters::GetColorForUsedCategory(int32 Index) const
-{
-	if (Index >= 0 && Index < sizeof(ColorPalette) / sizeof(ColorPalette[0]))
-	{
-		return ColorPalette[Index];
-	}
-
-	static bool bReateColorList = false;
-	static FColorList StaticColor;
-	if (!bReateColorList)
-	{
-		bReateColorList = true;
-		StaticColor.CreateColorMap();
-	}
-	return StaticColor.GetFColorByIndex(Index);
-}
-
-FLinearColor SVisualLoggerFilters::GetColorForUsedCategory(const FString& InFilterName) const
-{
-	int32 CategoryIndex = 0;
-	for (auto& CurrentFilter : Filters)
-	{
-		if (CurrentFilter->GetFilterNameAsString() == InFilterName)
-		{
-			break;
-		}
-		CategoryIndex++;
-	}
-
-	return GetColorForUsedCategory(CategoryIndex);
-}
-
 uint32 SVisualLoggerFilters::GetCategoryIndex(const FString& InFilterName) const
 {
 	uint32 CategoryIndex = 0;
@@ -362,7 +280,7 @@ void SVisualLoggerFilters::AddFilter(const FString& InFilterName)
 		}
 	}
 
-	const FLinearColor Color = GetColorForUsedCategory(Filters.Num());
+	const FLinearColor Color = FLogVisualizer::Get().GetColorForCategory(Filters.Num());
 	TSharedRef<SFilterWidget> NewFilter =
 		SNew(SFilterWidget)
 		.FilterName(*InFilterName)
@@ -412,7 +330,7 @@ void SVisualLoggerFilters::OnFiltersChanged()
 	Settings->CurrentPresets.CategoryFilters = EnabledFilters;
 	Settings->SaveConfig();
 
-	VisualLoggerInterface.Pin()->GetVisualLoggerEvents().OnFiltersChanged.ExecuteIfBound();
+	FLogVisualizer::Get().GetVisualLoggerEvents().OnFiltersChanged.ExecuteIfBound();
 }
 
 void SVisualLoggerFilters::AddFilter(const FString& GraphName, const FString& DataName)
diff --git a/Engine/Source/Developer/LogVisualizer/Private/SVisualLoggerFilters.h b/Engine/Source/Developer/LogVisualizer/Private/SVisualLoggerFilters.h
--- a/Engine/Source/Developer/LogVisualizer/Private/SVisualLoggerFilters.h
+++ b/Engine/Source/Developer/LogVisualizer/Private/SVisualLoggerFilters.h
@@ -27,11 +27,9 @@ public:
 	SLATE_BEGIN_ARGS(SVisualLoggerFilters) { }
 	SLATE_END_ARGS()
 		
-	void Construct(const FArguments& InArgs, const TSharedRef<FUICommandList>& InCommandList, TSharedPtr<IVisualLoggerInterface> VisualLoggerInterface);
+	void Construct(const FArguments& InArgs, const TSharedRef<FUICommandList>& InCommandList);
 	void AddFilter(const FString& InFilterName);
 	void AddFilter(const FString& GraphName, const FString& DataName);
-	FLinearColor GetColorForUsedCategory(int32 Index) const;
-	FLinearColor GetColorForUsedCategory(const FString& InFilterName) const;
 	uint32 GetCategoryIndex(const FString& InFilterName) const;
 	bool IsFilterEnabled(const FString& InFilterName, TEnumAsByte<ELogVerbosity::Type> Verbosity = ELogVerbosity::All);
 	bool IsFilterEnabled(const FString& InGraphName, const FString& InDataName, TEnumAsByte<ELogVerbosity::Type> Verbosity = ELogVerbosity::All);
@@ -53,10 +51,8 @@ public:
 protected:
 	/** The horizontal box which contains all the filters */
 	TSharedPtr<SWrapBox> FilterBox;
-	TWeakPtr<IVisualLoggerInterface> VisualLoggerInterface;
 	TArray<TSharedRef<SFilterWidget> > Filters;
 	TMap<FName, TArray<FSimpleGraphFilter> > GraphFilters;
 	TSharedPtr<SComboButton> GraphsFilterCombo;
 	FString GraphsFilter;
-	static FColor ColorPalette[];
 };
diff --git a/Engine/Source/Developer/LogVisualizer/Private/SVisualLoggerLogsList.cpp b/Engine/Source/Developer/LogVisualizer/Private/SVisualLoggerLogsList.cpp
--- a/Engine/Source/Developer/LogVisualizer/Private/SVisualLoggerLogsList.cpp
+++ b/Engine/Source/Developer/LogVisualizer/Private/SVisualLoggerLogsList.cpp
@@ -33,10 +33,8 @@ namespace ELogsSortMode
 	};
 }
 
-void SVisualLoggerLogsList::Construct(const FArguments& InArgs, const TSharedRef<FUICommandList>& InCommandList, TSharedPtr<IVisualLoggerInterface> InVisualLoggerInterface)
+void SVisualLoggerLogsList::Construct(const FArguments& InArgs, const TSharedRef<FUICommandList>& InCommandList)
 {
-	VisualLoggerInterface = InVisualLoggerInterface;
-
 	ChildSlot
 		[
 			SAssignNew(LogsLinesWidget, SListView<TSharedPtr<FLogEntryItem> >)
@@ -144,7 +142,7 @@ void SVisualLoggerLogsList::OnItemSelectionChanged(const FVisualLogDevice::FVisu
 
 	TArray<FVisualLoggerCategoryVerbosityPair> OutCategories;
 	FVisualLoggerHelpers::GetCategories(LogEntry.Entry, OutCategories);
-	if (VisualLoggerInterface.Pin()->HasValidCategories(OutCategories) == false)
+	if (FLogVisualizer::Get().GetVisualLoggerInterface()->HasValidCategories(OutCategories) == false)
 	{
 		return;
 	}
@@ -156,13 +154,13 @@ void SVisualLoggerLogsList::OnItemSelectionChanged(const FVisualLogDevice::FVisu
 		bool bShowLine = true;
 
 		FString CurrentCategory = LogLine->Category.ToString();
-		bShowLine = VisualLoggerInterface.Pin()->IsValidCategory(CurrentCategory, LogLine->Verbosity) /*&& (bHistogramGraphsFilter || (QuickFilterText.Len() == 0 || CurrentCategory.Find(QuickFilterText) != INDEX_NONE))*/;
+		bShowLine = FLogVisualizer::Get().GetVisualLoggerInterface()->IsValidCategory(CurrentCategory, LogLine->Verbosity) /*&& (bHistogramGraphsFilter || (QuickFilterText.Len() == 0 || CurrentCategory.Find(QuickFilterText) != INDEX_NONE))*/;
 
 		if (bShowLine)
 		{
 			FLogEntryItem EntryItem;
 			EntryItem.Category = LogLine->Category.ToString();
-			EntryItem.CategoryColor = VisualLoggerInterface.Pin()->GetCategoryColor(LogLine->Category);
+			EntryItem.CategoryColor = FLogVisualizer::Get().GetColorForCategory(LogLine->Category.ToString());
 			EntryItem.Verbosity = LogLine->Verbosity;
 			EntryItem.Line = LogLine->Line;
 			EntryItem.UserData = LogLine->UserData;
@@ -177,13 +175,13 @@ void SVisualLoggerLogsList::OnItemSelectionChanged(const FVisualLogDevice::FVisu
 		bool bShowLine = true;
 
 		FString CurrentCategory = Event.Name;
-		bShowLine = VisualLoggerInterface.Pin()->IsValidCategory(Event.Name, Event.Verbosity) /*&& (bHistogramGraphsFilter || (QuickFilterText.Len() == 0 || CurrentCategory.Find(QuickFilterText) != INDEX_NONE))*/;
+		bShowLine = FLogVisualizer::Get().GetVisualLoggerInterface()->IsValidCategory(Event.Name, Event.Verbosity) /*&& (bHistogramGraphsFilter || (QuickFilterText.Len() == 0 || CurrentCategory.Find(QuickFilterText) != INDEX_NONE))*/;
 
 		if (bShowLine)
 		{
 			FLogEntryItem EntryItem;
 			EntryItem.Category = Event.Name;
-			EntryItem.CategoryColor = VisualLoggerInterface.Pin()->GetCategoryColor(*EntryItem.Category);
+			EntryItem.CategoryColor = FLogVisualizer::Get().GetColorForCategory(*EntryItem.Category);
 			EntryItem.Verbosity = Event.Verbosity;
 			EntryItem.Line = FString::Printf(TEXT("Registered event: '%s' (%d times)%s"), *Event.Name, Event.Counter, Event.EventTags.Num() ? TEXT("\n") : TEXT(""));
 			for (auto& EventTag : Event.EventTags)
diff --git a/Engine/Source/Developer/LogVisualizer/Private/SVisualLoggerLogsList.h b/Engine/Source/Developer/LogVisualizer/Private/SVisualLoggerLogsList.h
--- a/Engine/Source/Developer/LogVisualizer/Private/SVisualLoggerLogsList.h
+++ b/Engine/Source/Developer/LogVisualizer/Private/SVisualLoggerLogsList.h
@@ -8,7 +8,7 @@ public:
 	SLATE_BEGIN_ARGS(SVisualLoggerLogsList){}
 	SLATE_END_ARGS()
 
-	void Construct(const FArguments& InArgs, const TSharedRef<FUICommandList>& InCommandList, TSharedPtr<IVisualLoggerInterface> VisualLoggerInterface);
+	void Construct(const FArguments& InArgs, const TSharedRef<FUICommandList>& InCommandList);
 
 	TSharedRef<ITableRow> LogEntryLinesGenerateRow(TSharedPtr<struct FLogEntryItem> Item, const TSharedRef<STableViewBase>& OwnerTable);
 	void LogEntryLineSelectionChanged(TSharedPtr<FLogEntryItem> SelectedItem, ESelectInfo::Type SelectInfo);
@@ -23,6 +23,5 @@ public:
 protected:
 	TSharedPtr<SListView<TSharedPtr<struct FLogEntryItem> > > LogsLinesWidget;
 	TArray<TSharedPtr<struct FLogEntryItem> > LogEntryLines;
-	TWeakPtr<IVisualLoggerInterface> VisualLoggerInterface;
 	FVisualLogDevice::FVisualLogEntryItem CurrentLogEntry;
 };
diff --git a/Engine/Source/Developer/LogVisualizer/Private/SVisualLoggerView.cpp b/Engine/Source/Developer/LogVisualizer/Private/SVisualLoggerView.cpp
--- a/Engine/Source/Developer/LogVisualizer/Private/SVisualLoggerView.cpp
+++ b/Engine/Source/Developer/LogVisualizer/Private/SVisualLoggerView.cpp
@@ -69,32 +69,29 @@ void SVisualLoggerView::GetTimelines(TArray<TSharedPtr<STimeline> >& OutList, bo
 	OutList = bOnlySelectedOnes ? TimelinesContainer->GetSelectedNodes() : TimelinesContainer->GetAllNodes();
 }
 
-void SVisualLoggerView::Construct(const FArguments& InArgs, const TSharedRef<FUICommandList>& InCommandList, TSharedPtr<IVisualLoggerInterface> InVisualLoggerInterface)
+void SVisualLoggerView::Construct(const FArguments& InArgs, const TSharedRef<FUICommandList>& InCommandList)
 {
-	VisualLoggerInterface = InVisualLoggerInterface;
 	AnimationOutlinerFillPercentage = .25f;
-	VisualLoggerEvents = VisualLoggerInterface.Pin()->GetVisualLoggerEvents();
+	VisualLoggerEvents = FLogVisualizer::Get().GetVisualLoggerEvents();
 
 	FVisualLoggerTimeSliderArgs TimeSliderArgs;
 	TimeSliderArgs.ViewRange = InArgs._ViewRange;
 	TimeSliderArgs.ClampMin = InArgs._ViewRange.Get().GetLowerBoundValue();
 	TimeSliderArgs.ClampMax = InArgs._ViewRange.Get().GetUpperBoundValue();
 	TimeSliderArgs.ScrubPosition = InArgs._ScrubPosition;
+	FLogVisualizer::Get().GetTimeSliderController()->SetTimesliderArgs(TimeSliderArgs);
 
 	TSharedRef<SScrollBar> ZoomScrollBar =
 		SNew(SScrollBar)
 		.Orientation(EOrientation::Orient_Horizontal)
 		.Thickness(FVector2D(2.0f, 2.0f));
 	ZoomScrollBar->SetState(0.0f, 1.0f);
-
-	TSharedPtr<FSequencerTimeSliderController> TimeSliderController(new FSequencerTimeSliderController(TimeSliderArgs));
-	TimeSliderController->SetExternalScrollbar(ZoomScrollBar);
-	VisualLoggerInterface.Pin()->SetTimeSliderController(TimeSliderController);
+	FLogVisualizer::Get().GetTimeSliderController()->SetExternalScrollbar(ZoomScrollBar);
 
 	// Create the top and bottom sliders
 	const bool bMirrorLabels = true;
-	TSharedRef<ITimeSlider> TopTimeSlider = SNew(STimeSlider, TimeSliderController.ToSharedRef()).MirrorLabels(bMirrorLabels);
-	TSharedRef<ITimeSlider> BottomTimeSlider = SNew(STimeSlider, TimeSliderController.ToSharedRef()).MirrorLabels(bMirrorLabels);
+	TSharedRef<ITimeSlider> TopTimeSlider = SNew(STimeSlider, FLogVisualizer::Get().GetTimeSliderController().ToSharedRef()).MirrorLabels(bMirrorLabels);
+	TSharedRef<ITimeSlider> BottomTimeSlider = SNew(STimeSlider, FLogVisualizer::Get().GetTimeSliderController().ToSharedRef()).MirrorLabels(bMirrorLabels);
 
 	TSharedRef<SScrollBar> ScrollBar =
 		SNew(SScrollBar)
@@ -200,10 +197,10 @@ void SVisualLoggerView::Construct(const FArguments& InArgs, const TSharedRef<FUI
 					+ SVerticalBox::Slot()
 					.FillHeight(1.0)
 					[
-						SNew(SInputCatcherOverlay, TimeSliderController.ToSharedRef())
+						SNew(SInputCatcherOverlay, FLogVisualizer::Get().GetTimeSliderController().ToSharedRef())
 						+ SOverlay::Slot()
 						[
-							MakeSectionOverlay(TimeSliderController.ToSharedRef(), InArgs._ViewRange, InArgs._ScrubPosition, false)
+							MakeSectionOverlay(FLogVisualizer::Get().GetTimeSliderController().ToSharedRef(), InArgs._ViewRange, InArgs._ScrubPosition, false)
 						]
 						+ SOverlay::Slot()
 						[
@@ -211,13 +208,12 @@ void SVisualLoggerView::Construct(const FArguments& InArgs, const TSharedRef<FUI
 							.ExternalScrollbar(ScrollBar)
 							+ SScrollBox::Slot()
 							[
-								SAssignNew(TimelinesContainer, STimelinesContainer, SharedThis(this), TimeSliderController.ToSharedRef())
-								.VisualLoggerInterface(VisualLoggerInterface)
+								SAssignNew(TimelinesContainer, STimelinesContainer, SharedThis(this), FLogVisualizer::Get().GetTimeSliderController().ToSharedRef())
 							]
 						]
 						+ SOverlay::Slot()
 						[
-							MakeSectionOverlay(TimeSliderController.ToSharedRef(), InArgs._ViewRange, InArgs._ScrubPosition, true)
+							MakeSectionOverlay(FLogVisualizer::Get().GetTimeSliderController().ToSharedRef(), InArgs._ViewRange, InArgs._ScrubPosition, true)
 						]
 						+ SOverlay::Slot()
 						.VAlign(VAlign_Bottom)
@@ -271,12 +267,6 @@ void SVisualLoggerView::Construct(const FArguments& InArgs, const TSharedRef<FUI
 		];
 
 		SearchBox->SetText(FText::FromString(CurrentPresets.DataFilter));
-		//ScrollBox->AddSlot()
-		//[
-		//	SAssignNew(TimelinesContainer, STimelinesContainer, SharedThis(this), TimeSliderController.ToSharedRef())
-		//	.VisualLoggerInterface(VisualLoggerInterface)
-		//];
-
 }
 
 void SVisualLoggerView::SetAnimationOutlinerFillPercentage(float FillPercentage) 
@@ -358,7 +348,7 @@ void SVisualLoggerView::OnFiltersSearchChanged(const FText& Filter)
 
 FCursorReply SVisualLoggerView::OnCursorQuery(const FGeometry& MyGeometry, const FPointerEvent& CursorEvent) const
 {
-	if (VisualLoggerInterface.Pin()->GetTimeSliderController()->IsPanning())
+	if (FLogVisualizer::Get().GetTimeSliderController()->IsPanning())
 	{
 		return FCursorReply::Cursor(EMouseCursor::GrabHand);
 	}
diff --git a/Engine/Source/Developer/LogVisualizer/Private/SVisualLoggerView.h b/Engine/Source/Developer/LogVisualizer/Private/SVisualLoggerView.h
--- a/Engine/Source/Developer/LogVisualizer/Private/SVisualLoggerView.h
+++ b/Engine/Source/Developer/LogVisualizer/Private/SVisualLoggerView.h
@@ -16,7 +16,7 @@ public:
 		SLATE_EVENT(FOnFiltersSearchChanged, OnFiltersSearchChanged)
 	SLATE_END_ARGS();
 
-	void Construct(const FArguments& InArgs, const TSharedRef<FUICommandList>& InCommandList, TSharedPtr<IVisualLoggerInterface> VisualLoggerInterface);
+	void Construct(const FArguments& InArgs, const TSharedRef<FUICommandList>& InCommandList);
 	float GetAnimationOutlinerFillPercentage() const { 
 		SSplitter::FSlot const& LeftSplitterSlot = SearchSplitter->SlotAt(0);
 		SSplitter::FSlot const& RightSplitterSlot = SearchSplitter->SlotAt(1);
@@ -42,7 +42,6 @@ public:
 	void ResetData();
 
 protected:
-	TWeakPtr<struct IVisualLoggerInterface> VisualLoggerInterface;
 	TSharedPtr<class STimelinesContainer> TimelinesContainer;
 	TSharedPtr<class SSplitter> SearchSplitter;
 	TSharedPtr<class SScrollBox> ScrollBox;
diff --git a/Engine/Source/Developer/LogVisualizer/Private/TimeSliderController.cpp b/Engine/Source/Developer/LogVisualizer/Private/TimeSliderController.cpp
--- a/Engine/Source/Developer/LogVisualizer/Private/TimeSliderController.cpp
+++ b/Engine/Source/Developer/LogVisualizer/Private/TimeSliderController.cpp
@@ -103,6 +103,11 @@ FSequencerTimeSliderController::FSequencerTimeSliderController(const FVisualLogg
 	CursorBackground = FEditorStyle::GetBrush("Sequencer.SectionArea.Background");
 }
 
+void FSequencerTimeSliderController::SetTimesliderArgs(const FVisualLoggerTimeSliderArgs& InArgs)
+{
+	TimeSliderArgs = InArgs;
+}
+
 struct FDrawTickArgs
 {
 	/** Geometry of the area */
diff --git a/Engine/Source/Developer/LogVisualizer/Private/TimeSliderController.h b/Engine/Source/Developer/LogVisualizer/Private/TimeSliderController.h
--- a/Engine/Source/Developer/LogVisualizer/Private/TimeSliderController.h
+++ b/Engine/Source/Developer/LogVisualizer/Private/TimeSliderController.h
@@ -24,6 +24,7 @@ class FSequencerTimeSliderController : public ITimeSliderController
 {
 public:
 	FSequencerTimeSliderController(const FVisualLoggerTimeSliderArgs& InArgs);
+	void SetTimesliderArgs(const FVisualLoggerTimeSliderArgs& InArgs);
 
 	/** ITimeSliderController Interface */
 	virtual int32 OnPaintTimeSlider( bool bMirrorLabels, const FGeometry& AllottedGeometry, const FSlateRect& MyClippingRect, FSlateWindowElementList& OutDrawElements, int32 LayerId, const FWidgetStyle& InWidgetStyle, bool bParentEnabled ) const override;
diff --git a/Engine/Source/Developer/LogVisualizer/Private/VisualLoggerCanvasRenderer.cpp b/Engine/Source/Developer/LogVisualizer/Private/VisualLoggerCanvasRenderer.cpp
--- a/Engine/Source/Developer/LogVisualizer/Private/VisualLoggerCanvasRenderer.cpp
+++ b/Engine/Source/Developer/LogVisualizer/Private/VisualLoggerCanvasRenderer.cpp
@@ -69,7 +69,7 @@ void FVisualLoggerCanvasRenderer::DrawOnCanvas(class UCanvas* Canvas, class APla
 		return;
 	}
 
-	UWorld* World = VisualLoggerInterface.Pin()->GetWorld();
+	UWorld* World = FLogVisualizer::Get().GetWorld();
 	if (World == NULL)
 	{
 		return;
@@ -84,7 +84,7 @@ void FVisualLoggerCanvasRenderer::DrawOnCanvas(class UCanvas* Canvas, class APla
 	for (const auto CurrentData : SelectedEntry.DataBlocks)
 	{
 		const FName TagName = CurrentData.TagName;
-		const bool bIsValidByFilter = VisualLoggerInterface.Pin()->IsValidCategory(CurrentData.Category.ToString(), ELogVerbosity::All) && VisualLoggerInterface.Pin()->IsValidCategory(CurrentData.TagName.ToString(), ELogVerbosity::All);
+		const bool bIsValidByFilter = FLogVisualizer::Get().GetVisualLoggerInterface()->IsValidCategory(CurrentData.Category.ToString(), ELogVerbosity::All) && FLogVisualizer::Get().GetVisualLoggerInterface()->IsValidCategory(CurrentData.TagName.ToString(), ELogVerbosity::All);
 		FVisualLogExtensionInterface* Extension = FVisualLogger::Get().GetExtensionForTag(TagName);
 		if (!Extension)
 		{
@@ -93,11 +93,11 @@ void FVisualLoggerCanvasRenderer::DrawOnCanvas(class UCanvas* Canvas, class APla
 
 		if (!bIsValidByFilter)
 		{
-			Extension->DisableDrawingForData(VisualLoggerInterface.Pin()->GetWorld(), Canvas, NULL, TagName, CurrentData, SelectedEntry.TimeStamp);
+			Extension->DisableDrawingForData(FLogVisualizer::Get().GetWorld(), Canvas, NULL, TagName, CurrentData, SelectedEntry.TimeStamp);
 		}
 		else
 		{
-			Extension->DrawData(VisualLoggerInterface.Pin()->GetWorld(), Canvas, NULL, TagName, CurrentData, SelectedEntry.TimeStamp);
+			Extension->DrawData(FLogVisualizer::Get().GetWorld(), Canvas, NULL, TagName, CurrentData, SelectedEntry.TimeStamp);
 		}
 	}
 
@@ -125,14 +125,14 @@ void FVisualLoggerCanvasRenderer::DrawHistogramGraphs(class UCanvas* Canvas, cla
 		TMap<FName, FGraphLineData> GraphLines;
 	};
 
-	if (VisualLoggerInterface.Pin()->GetTimeSliderController() == NULL)
+	if (FLogVisualizer::Get().GetTimeSliderController().IsValid() == false)
 	{
 		return;
 	}
 
 	TMap<FName, FGraphData>	CollectedGraphs;
 
-	const FVisualLoggerTimeSliderArgs&  TimeSliderArgs = VisualLoggerInterface.Pin()->GetTimeSliderController()->GetTimeSliderArgs();
+	const FVisualLoggerTimeSliderArgs&  TimeSliderArgs = FLogVisualizer::Get().GetTimeSliderController()->GetTimeSliderArgs();
 	TRange<float> LocalViewRange = TimeSliderArgs.ViewRange.Get();
 	const float LocalViewRangeMin = LocalViewRange.GetLowerBoundValue();
 	const float LocalViewRangeMax = LocalViewRange.GetUpperBoundValue();
@@ -169,7 +169,7 @@ void FVisualLoggerCanvasRenderer::DrawHistogramGraphs(class UCanvas* Canvas, cla
 			const FName CurrentGraphName = CurrentSample.GraphName;
 			const FName CurrentDataName = CurrentSample.DataName;
 
-			const bool bIsValidByFilter = VisualLoggerInterface.Pin()->IsValidCategory(CurrentSample.GraphName.ToString(), CurrentSample.DataName.ToString(), ELogVerbosity::All);
+			const bool bIsValidByFilter = FLogVisualizer::Get().GetVisualLoggerInterface()->IsValidCategory(CurrentSample.GraphName.ToString(), CurrentSample.DataName.ToString(), ELogVerbosity::All);
 			if (bIsValidByFilter)
 			{
 				FGraphData &GraphData = CollectedGraphs.FindOrAdd(CurrentSample.GraphName);
@@ -199,7 +199,7 @@ void FVisualLoggerCanvasRenderer::DrawHistogramGraphs(class UCanvas* Canvas, cla
 			const FName CurrentGraphName = CurrentSample.GraphName;
 			const FName CurrentDataName = CurrentSample.DataName;
 
-			const bool bIsValidByFilter = VisualLoggerInterface.Pin()->IsValidCategory(CurrentSample.GraphName.ToString(), CurrentSample.DataName.ToString(), ELogVerbosity::All);
+			const bool bIsValidByFilter = FLogVisualizer::Get().GetVisualLoggerInterface()->IsValidCategory(CurrentSample.GraphName.ToString(), CurrentSample.DataName.ToString(), ELogVerbosity::All);
 			if (bIsValidByFilter)
 			{
 				FGraphData &GraphData = CollectedGraphs.FindOrAdd(CurrentSample.GraphName);
diff --git a/Engine/Source/Developer/LogVisualizer/Private/VisualLoggerCanvasRenderer.h b/Engine/Source/Developer/LogVisualizer/Private/VisualLoggerCanvasRenderer.h
--- a/Engine/Source/Developer/LogVisualizer/Private/VisualLoggerCanvasRenderer.h
+++ b/Engine/Source/Developer/LogVisualizer/Private/VisualLoggerCanvasRenderer.h
@@ -3,8 +3,7 @@
 
 struct FVisualLoggerCanvasRenderer
 {
-	FVisualLoggerCanvasRenderer(TSharedPtr<IVisualLoggerInterface> InVisualLoggerInterface) 
-		: VisualLoggerInterface(InVisualLoggerInterface) 
+	FVisualLoggerCanvasRenderer() 
 	{}
 
 	void DrawOnCanvas(class UCanvas* Canvas, class APlayerController*);
@@ -16,7 +15,6 @@ protected:
 
 private:
 	FVisualLogEntry SelectedEntry;
-	TWeakPtr<IVisualLoggerInterface> VisualLoggerInterface;
 	TMap<FString, TArray<FString> > UsedGraphCategories;
 	TWeakPtr<class STimeline> CurrentTimeLine;
 };
diff --git a/Engine/Source/Developer/LogVisualizer/Private/VisualLoggerCommands.h b/Engine/Source/Developer/LogVisualizer/Private/VisualLoggerCommands.h
--- a/Engine/Source/Developer/LogVisualizer/Private/VisualLoggerCommands.h
+++ b/Engine/Source/Developer/LogVisualizer/Private/VisualLoggerCommands.h
@@ -30,6 +30,9 @@ public:
 		UI_COMMAND(ToggleGraphs, "ToggleGraphs", "Toggle graphs visualization on/off", EUserInterfaceActionType::ToggleButton, FInputGesture());
 		UI_COMMAND(ResetData, "Clear", "Clear all data", EUserInterfaceActionType::Button, FInputGesture());
 		UI_COMMAND(HideEmpty, "HideEmpty", "Hide logs without text information", EUserInterfaceActionType::ToggleButton, FInputGesture());
+
+		UI_COMMAND(MoveCursorLeft, "MoveCursorLeft", "Move cursor left on timeline", EUserInterfaceActionType::Button, FInputGesture(EKeys::Left));
+		UI_COMMAND(MoveCursorRight, "MoveCursorLeft", "Move cursor right on timeline", EUserInterfaceActionType::Button, FInputGesture(EKeys::Right));
 	}
 
 public:
@@ -44,6 +47,8 @@ public:
 	TSharedPtr<FUICommandInfo> ToggleGraphs;
 	TSharedPtr<FUICommandInfo> ResetData;
 	TSharedPtr<FUICommandInfo> HideEmpty;
+	TSharedPtr<FUICommandInfo> MoveCursorLeft;
+	TSharedPtr<FUICommandInfo> MoveCursorRight;
 };
 
 
diff --git a/Engine/Source/Developer/LogVisualizer/Private/VisualLoggerRenderingActor.cpp b/Engine/Source/Developer/LogVisualizer/Private/VisualLoggerRenderingActor.cpp
--- a/Engine/Source/Developer/LogVisualizer/Private/VisualLoggerRenderingActor.cpp
+++ b/Engine/Source/Developer/LogVisualizer/Private/VisualLoggerRenderingActor.cpp
@@ -180,7 +180,7 @@ void AVisualLoggerRenderingActor::ObjectSelectionChanged(TSharedPtr<class STimel
 	}
 }
 
-void AVisualLoggerRenderingActor::OnItemSelectionChanged(const FVisualLogDevice::FVisualLogEntryItem& EntryItem, TSharedPtr<IVisualLoggerInterface> VisualLoggerInterface)
+void AVisualLoggerRenderingActor::OnItemSelectionChanged(const FVisualLogDevice::FVisualLogEntryItem& EntryItem)
 {
 	const FVisualLogEntry* Entry = &EntryItem.Entry;
 	const FVisualLogShapeElement* ElementToDraw = Entry->ElementsToDraw.GetData();
@@ -225,7 +225,7 @@ void AVisualLoggerRenderingActor::OnItemSelectionChanged(const FVisualLogDevice:
 	for (const auto CurrentData : Entry->DataBlocks)
 	{
 		const FName TagName = CurrentData.TagName;
-		const bool bIsValidByFilter = VisualLoggerInterface->IsValidCategory(CurrentData.Category.ToString(), ELogVerbosity::All) && VisualLoggerInterface->IsValidCategory(CurrentData.TagName.ToString(), ELogVerbosity::All);
+		const bool bIsValidByFilter = FLogVisualizer::Get().GetVisualLoggerInterface()->IsValidCategory(CurrentData.Category.ToString(), ELogVerbosity::All) && FLogVisualizer::Get().GetVisualLoggerInterface()->IsValidCategory(CurrentData.TagName.ToString(), ELogVerbosity::All);
 		FVisualLogExtensionInterface* Extension = FVisualLogger::Get().GetExtensionForTag(TagName);
 		if (!Extension)
 		{
@@ -244,7 +244,7 @@ void AVisualLoggerRenderingActor::OnItemSelectionChanged(const FVisualLogDevice:
 
 	for (int32 ElementIndex = 0; ElementIndex < ElementsCount; ++ElementIndex, ++ElementToDraw)
 	{
-		if (VisualLoggerInterface.IsValid() && !VisualLoggerInterface->IsValidCategory(ElementToDraw->Category.ToString(), ElementToDraw->Verbosity))
+		if (!FLogVisualizer::Get().GetVisualLoggerInterface()->IsValidCategory(ElementToDraw->Category.ToString(), ElementToDraw->Verbosity))
 		{
 			continue;
 		}
diff --git a/Engine/Source/Developer/LogVisualizer/Private/VisualLoggerRenderingActor.h b/Engine/Source/Developer/LogVisualizer/Private/VisualLoggerRenderingActor.h
--- a/Engine/Source/Developer/LogVisualizer/Private/VisualLoggerRenderingActor.h
+++ b/Engine/Source/Developer/LogVisualizer/Private/VisualLoggerRenderingActor.h
@@ -16,7 +16,7 @@ class LOGVISUALIZER_API AVisualLoggerRenderingActor : public AActor
 public:
 	GENERATED_UCLASS_BODY()
 
-	void OnItemSelectionChanged(const FVisualLogDevice::FVisualLogEntryItem& EntryItem, TSharedPtr<IVisualLoggerInterface> VisualLoggerInterface);
+	void OnItemSelectionChanged(const FVisualLogDevice::FVisualLogEntryItem& EntryItem);
 	void ObjectSelectionChanged(TSharedPtr<class STimeline> TimeLine);
 
 private:
diff --git a/Engine/Source/Developer/LogVisualizer/Public/ILogVisualizer.h b/Engine/Source/Developer/LogVisualizer/Public/ILogVisualizer.h
--- a/Engine/Source/Developer/LogVisualizer/Public/ILogVisualizer.h
+++ b/Engine/Source/Developer/LogVisualizer/Public/ILogVisualizer.h
@@ -2,21 +2,10 @@
 
 #pragma once
 
-class ILogVisualizer
+class ILogVisualizer : public IModuleInterface, public IModularFeature
 {
 public:
-	/** Virtual destructor */
-	virtual ~ILogVisualizer() {}
-
-	/** spawns LogVisualizer's UI */
-	virtual void SummonUI(UWorld* InWorld) = 0;
-
-	/** closes LogVisualizer's UI */
-	virtual void CloseUI(UWorld* InWorld) = 0;
-
-	/** checks if we have LogVisualizer UI opened for given World */
-	virtual bool IsOpenUI(UWorld* InWorld) = 0;
-
-	/** Helper actor used to render data in given world */
-	virtual class AActor* GetHelperActor(class UWorld* InWorld) = 0;
+	virtual void Goto(float Timestamp, FName LogOwner = NAME_None) = 0;
+	virtual void GotoNextItem() = 0;
+	virtual void GotoPreviousItem() = 0;
 };
\ No newline at end of file
diff --git a/Engine/Source/Developer/LogVisualizer/Public/LogVisualizerModule.h b/Engine/Source/Developer/LogVisualizer/Public/LogVisualizerModule.h
deleted file mode 100644
diff --git a/Engine/Source/Runtime/Engine/Private/VisualLogger/VisualLogger.cpp b/Engine/Source/Runtime/Engine/Private/VisualLogger/VisualLogger.cpp
--- a/Engine/Source/Runtime/Engine/Private/VisualLogger/VisualLogger.cpp
+++ b/Engine/Source/Runtime/Engine/Private/VisualLogger/VisualLogger.cpp
@@ -187,7 +187,6 @@ FCustomVersionRegistration GVisualLoggerVersion(EVisualLoggerVersion::GUID, EVis
 
 #if !(UE_BUILD_SHIPPING || UE_BUILD_TEST)
 
-#include "Developer/LogVisualizer/Public/LogVisualizerModule.h"
 #if WITH_EDITOR
 #include "SlateBasics.h"
 #endif
