Fixed crash in LogVisualizer.

diff --git a/Engine/Source/Developer/LogVisualizer/Private/SVisualLogger.cpp b/Engine/Source/Developer/LogVisualizer/Private/SVisualLogger.cpp
--- a/Engine/Source/Developer/LogVisualizer/Private/SVisualLogger.cpp
+++ b/Engine/Source/Developer/LogVisualizer/Private/SVisualLogger.cpp
@@ -92,7 +92,7 @@ void SVisualLogger::FVisualLoggerDevice::Serialize(const class UObject* LogOwner
 		return;
 	}
 
-	UWorld* World = LogOwner != nullptr ? GEngine->GetWorldFromContextObject(LogOwner) : NULL;
+	UWorld* World = FLogVisualizer::Get().GetWorld();
 	if (LastWorld.Get() != World)
 	{
 		Owner->OnNewWorld(LastWorld.Get());
