Fixes for visual logs and LogVisualizer:
- Fixed compilation warnings,
- fixes bug in vlogs serialization,
- fixed information about warnings and errors on timelines in LogVisualizer,
- fixed data reset in LogVisualizer

diff --git a/Engine/Source/Developer/LogVisualizer/Private/SFilterWidget.h b/Engine/Source/Developer/LogVisualizer/Private/SFilterWidget.h
--- a/Engine/Source/Developer/LogVisualizer/Private/SFilterWidget.h
+++ b/Engine/Source/Developer/LogVisualizer/Private/SFilterWidget.h
@@ -72,9 +72,9 @@ public:
 
 private:
 	/** Handler for when the filter checkbox is clicked */
-	void FilterToggled(ESlateCheckBoxState::Type NewState)
+	void FilterToggled(ECheckBoxState NewState)
 	{
-		bEnabled = NewState == ESlateCheckBoxState::Checked;
+		bEnabled = NewState == ECheckBoxState::Checked;
 		OnFilterChanged.ExecuteIfBound();
 	}
 
@@ -126,10 +126,10 @@ private:
 	void RemoveAllFilters() { OnRequestRemoveAll.ExecuteIfBound(); }
 
 	/** Handler to determine the "checked" state of the filter checkbox */
-	ESlateCheckBoxState::Type IsChecked() const { return bEnabled ? ESlateCheckBoxState::Checked : ESlateCheckBoxState::Unchecked; }
+	ECheckBoxState IsChecked() const { return bEnabled ? ECheckBoxState::Checked : ECheckBoxState::Unchecked; }
 
 	/** Handler to determine the color of the checkbox when it is checked */
-	FSlateColor GetFilterForegroundColor() const { return IsChecked() == ESlateCheckBoxState::Checked ? FilterColor : FLinearColor::White; }
+	FSlateColor GetFilterForegroundColor() const { return IsChecked() == ECheckBoxState::Checked ? FilterColor : FLinearColor::White; }
 
 	/** Handler to determine the padding of the checkbox text when it is pressed */
 	FMargin GetFilterNamePadding() const { return ToggleButtonPtr->IsPressed() ? FMargin(3, 2, 4, 0) : FMargin(3, 1, 4, 1); }
@@ -138,7 +138,7 @@ private:
 	FSlateColor GetFilterNameColorAndOpacity() const
 	{
 		const float DimFactor = 0.75f;
-		return IsChecked() == ESlateCheckBoxState::Checked ? (IsHovered() ? ColorCategory * DimFactor : ColorCategory) : (IsHovered() ? FLinearColor::White : FLinearColor::White * DimFactor);
+		return IsChecked() == ECheckBoxState::Checked ? (IsHovered() ? ColorCategory * DimFactor : ColorCategory) : (IsHovered() ? FLinearColor::White : FLinearColor::White * DimFactor);
 	}
 
 	FSlateColor GetBorderBackgroundColor() const
diff --git a/Engine/Source/Developer/LogVisualizer/Private/STimelineBar.cpp b/Engine/Source/Developer/LogVisualizer/Private/STimelineBar.cpp
--- a/Engine/Source/Developer/LogVisualizer/Private/STimelineBar.cpp
+++ b/Engine/Source/Developer/LogVisualizer/Private/STimelineBar.cpp
@@ -393,16 +393,22 @@ int32 STimelineBar::OnPaint(const FPaintArgs& Args, const FGeometry& AllottedGeo
 			}
 
 			const TArray<FVisualLogLine>& LogLines = CurrentEntry.LogLines;
+			bool bAddedWarning = false;
+			bool bAddedError = false;
 			for (const FVisualLogLine& CurrentLine : LogLines)
 			{
-				if (CurrentLine.Verbosity <= ELogVerbosity::Error)
+				if (CurrentLine.Verbosity <= ELogVerbosity::Error && !bAddedError)
 				{
 					ErrorTimes.AddUnique(CurrentEntry.TimeStamp);
-					break;
+					bAddedError = true;
 				}
-				else if (CurrentLine.Verbosity == ELogVerbosity::Warning)
+				else if (CurrentLine.Verbosity == ELogVerbosity::Warning && !bAddedWarning)
 				{
 					WarningTimes.AddUnique(CurrentEntry.TimeStamp);
+					bAddedWarning = true;
+				}
+				if (bAddedError && bAddedWarning)
+				{
 					break;
 				}
 			}
diff --git a/Engine/Source/Developer/LogVisualizer/Private/SVisualLogger.cpp b/Engine/Source/Developer/LogVisualizer/Private/SVisualLogger.cpp
--- a/Engine/Source/Developer/LogVisualizer/Private/SVisualLogger.cpp
+++ b/Engine/Source/Developer/LogVisualizer/Private/SVisualLogger.cpp
@@ -719,6 +719,30 @@ void SVisualLogger::ResetData()
 	{
 		VisualLoggerFilters->ResetData();
 	}
+
+	if (LogsList.IsValid())
+	{
+		LogsList->OnItemSelectionChanged(FVisualLogDevice::FVisualLogEntryItem());
+	}
+
+	if (StatusView.IsValid())
+	{
+		StatusView->OnItemSelectionChanged(FVisualLogDevice::FVisualLogEntryItem());
+	}
+
+	if (VisualLoggerCanvasRenderer.IsValid())
+	{
+		VisualLoggerCanvasRenderer->OnItemSelectionChanged(FVisualLogEntry());
+		VisualLoggerCanvasRenderer->ObjectSelectionChanged(NULL);
+	}
+
+	AVisualLoggerRenderingActor* HelperActor = Cast<AVisualLoggerRenderingActor>(FLogVisualizer::Get().GetVisualLoggerHelperActor());
+	if (HelperActor)
+	{
+		HelperActor->OnItemSelectionChanged(FVisualLogDevice::FVisualLogEntryItem());
+		HelperActor->ObjectSelectionChanged(NULL);
+	}
+
 }
 
 void SVisualLogger::OnNewWorld(UWorld* NewWorld)
diff --git a/Engine/Source/Developer/LogVisualizer/Private/VisualLoggerRenderingActor.cpp b/Engine/Source/Developer/LogVisualizer/Private/VisualLoggerRenderingActor.cpp
--- a/Engine/Source/Developer/LogVisualizer/Private/VisualLoggerRenderingActor.cpp
+++ b/Engine/Source/Developer/LogVisualizer/Private/VisualLoggerRenderingActor.cpp
@@ -165,6 +165,7 @@ void AVisualLoggerRenderingActor::ObjectSelectionChanged(TSharedPtr<class STimel
 		Texts.Reset();
 		Cylinders.Reset();
 		Capsles.Reset();
+		LogEntriesPath.Reset();
 		MarkComponentsRenderStateDirty();
 		return;
 	}
@@ -200,7 +201,7 @@ void AVisualLoggerRenderingActor::OnItemSelectionChanged(const FVisualLogDevice:
 #endif
 
 	const float Length = 100;
-	const FVector DirectionNorm = FVector(0, 0, 1).SafeNormal();
+	const FVector DirectionNorm = FVector(0, 0, 1).GetSafeNormal();
 	FVector YAxis, ZAxis;
 	DirectionNorm.FindBestAxisVectors(YAxis, ZAxis);
 	Cones.Add(FDebugRenderSceneProxy::FCone(FScaleMatrix(FVector(Length)) * FMatrix(DirectionNorm, YAxis, ZAxis, Entry->Location), 5, 5, FColor::Red));
@@ -375,7 +376,7 @@ void AVisualLoggerRenderingActor::OnItemSelectionChanged(const FVisualLogDevice:
 			for (int32 Index = 0; Index + 2 < ElementToDraw->Points.Num(); Index += 3)
 			{
 				const FVector Orgin = ElementToDraw->Points[Index];
-				const FVector Direction = ElementToDraw->Points[Index + 1].SafeNormal();
+				const FVector Direction = ElementToDraw->Points[Index + 1].GetSafeNormal();
 				const FVector Angles = ElementToDraw->Points[Index + 2];
 				const float Length = Angles.X;
 
diff --git a/Engine/Source/Runtime/Engine/Public/VisualLogger/VisualLogger.inl b/Engine/Source/Runtime/Engine/Public/VisualLogger/VisualLogger.inl
--- a/Engine/Source/Runtime/Engine/Public/VisualLogger/VisualLogger.inl
+++ b/Engine/Source/Runtime/Engine/Public/VisualLogger/VisualLogger.inl
@@ -76,6 +76,18 @@ FVisualLogEntry* FVisualLogger::GetEntryToWrite(const class UObject* Object, flo
 		if (World)
 		{
 			World->GetTimerManager().ClearAllTimersForObject(this);
+			for (auto& CurrentPair : CurrentEntryPerObject)
+			{
+				FVisualLogEntry* Entry = &CurrentPair.Value;
+				if (Entry->TimeStamp >= 0 && Entry->TimeStamp < TimeStamp)
+				{
+					for (auto* Device : OutputDevices)
+					{
+						Device->Serialize(CurrentPair.Key, ObjectToNameMap[CurrentPair.Key], ObjectToClassNameMap[CurrentPair.Key], *Entry);
+					}
+					Entry->Reset();
+				}
+			}
 		}
 	}
 
