Fixed tree expansion for status view tab in LogVisualizer.

diff --git a/Engine/Source/Developer/LogVisualizer/Private/SVisualLoggerStatusView.cpp b/Engine/Source/Developer/LogVisualizer/Private/SVisualLoggerStatusView.cpp
--- a/Engine/Source/Developer/LogVisualizer/Private/SVisualLoggerStatusView.cpp
+++ b/Engine/Source/Developer/LogVisualizer/Private/SVisualLoggerStatusView.cpp
@@ -63,8 +63,15 @@ void GenerateChildren(TSharedPtr<FLogStatusItem> StatusItem, const FVisualLogSta
 }
 void SVisualLoggerStatusView::OnItemSelectionChanged(const FVisualLogDevice::FVisualLogEntryItem& LogEntry)
 {
-	//TSet<TSharedPtr<FLogStatusItem> > ExpandedItems;
-	//StatusItemsView->GetExpandedItems(ExpandedItems);
+	TArray<FString> ExpandedCategories;
+	for (int32 ItemIndex = 0; ItemIndex < StatusItems.Num(); ItemIndex++)
+	{
+		const bool bIsExpanded = StatusItemsView->IsItemExpanded(StatusItems[ItemIndex]);
+		if (bIsExpanded)
+		{
+			ExpandedCategories.Add(StatusItems[ItemIndex]->ItemText);
+		}
+	}
 
 	StatusItems.Empty();
 
@@ -101,10 +108,17 @@ void SVisualLoggerStatusView::OnItemSelectionChanged(const FVisualLogDevice::FVi
 
 	StatusItemsView->RequestTreeRefresh();
 
-	//for (auto& CurrentItem : ExpandedItems)
-	//{
-	//	StatusItemsView->SetItemExpansion(CurrentItem, true);
-	//}
+	for (int32 ItemIndex = 0; ItemIndex < StatusItems.Num(); ItemIndex++)
+	{
+		for (const FString& Category : ExpandedCategories)
+		{
+			if (StatusItems[ItemIndex]->ItemText == Category)
+			{
+				StatusItemsView->SetItemExpansion(StatusItems[ItemIndex], true);
+				break;
+			}
+		}
+	}
 }
 
 void SVisualLoggerStatusView::OnLogStatusGetChildren(TSharedPtr<FLogStatusItem> InItem, TArray< TSharedPtr<FLogStatusItem> >& OutItems)
