Fixed issue when moving cursor on timelines outside time rande in LogVisualizer.

diff --git a/Engine/Source/Developer/LogVisualizer/Private/TimeSliderController.cpp b/Engine/Source/Developer/LogVisualizer/Private/TimeSliderController.cpp
--- a/Engine/Source/Developer/LogVisualizer/Private/TimeSliderController.cpp
+++ b/Engine/Source/Developer/LogVisualizer/Private/TimeSliderController.cpp
@@ -453,8 +453,17 @@ FReply FSequencerTimeSliderController::OnMouseMove( TSharedRef<SWidget> WidgetOw
 				FScrubRangeToScreen RangeToScreen( TimeSliderArgs.ViewRange.Get(), MyGeometry.Size );
 				FVector2D CursorPos = MyGeometry.AbsoluteToLocal( MouseEvent.GetLastScreenSpacePosition() );
 				float NewValue = RangeToScreen.LocalXToInput( CursorPos.X );
+				if (TimeSliderArgs.ClampMin.Get().IsSet() && NewValue < TimeSliderArgs.ClampMin.Get().GetValue())
+				{
+					NewValue = TimeSliderArgs.ClampMin.Get().GetValue();
+				}
+
+				if (TimeSliderArgs.ClampMax.Get().IsSet() && NewValue > TimeSliderArgs.ClampMax.Get().GetValue())
+				{
+					NewValue = TimeSliderArgs.ClampMax.Get().GetValue();
+				}
 
-				CommitScrubPosition( NewValue, /*bIsScrubbing=*/true );
+				CommitScrubPosition(NewValue, /*bIsScrubbing=*/true);
 			}
 		}
 		return FReply::Handled();
