Avoid symbol collision between two FSequencerTimeSliderControllers.

- Slightly different solution than PR #778 (lesser impact on code).

diff --git a/Engine/Source/Developer/LogVisualizer/Private/LogVisualizer.cpp b/Engine/Source/Developer/LogVisualizer/Private/LogVisualizer.cpp
--- a/Engine/Source/Developer/LogVisualizer/Private/LogVisualizer.cpp
+++ b/Engine/Source/Developer/LogVisualizer/Private/LogVisualizer.cpp
@@ -67,7 +67,7 @@ FColor FLogVisualizerColorPalette[] = {
 void FLogVisualizer::Initialize()
 {
 	StaticInstance = MakeShareable(new FLogVisualizer);
-	Get().TimeSliderController = MakeShareable(new FSequencerTimeSliderController(FVisualLoggerTimeSliderArgs()));
+	Get().TimeSliderController = MakeShareable(new FVisualLoggerTimeSliderController(FVisualLoggerTimeSliderArgs()));
 
 }
 
diff --git a/Engine/Source/Developer/LogVisualizer/Private/LogVisualizer.h b/Engine/Source/Developer/LogVisualizer/Private/LogVisualizer.h
--- a/Engine/Source/Developer/LogVisualizer/Private/LogVisualizer.h
+++ b/Engine/Source/Developer/LogVisualizer/Private/LogVisualizer.h
@@ -42,7 +42,7 @@ struct IVisualLoggerInterface
 	virtual bool IsValidCategory(const FString& InGraphName, const FString& InDataName, TEnumAsByte<ELogVerbosity::Type> Verbosity = ELogVerbosity::All) = 0;
 };
 
-class FSequencerTimeSliderController;
+class FVisualLoggerTimeSliderController;
 struct LOGVISUALIZER_API FLogVisualizer
 {
 	/** LogVisualizer interface*/
@@ -52,7 +52,7 @@ struct LOGVISUALIZER_API FLogVisualizer
 
 	FLinearColor GetColorForCategory(int32 Index) const;
 	FLinearColor GetColorForCategory(const FString& InFilterName) const;
-	TSharedPtr<FSequencerTimeSliderController> GetTimeSliderController() { return TimeSliderController; }
+	TSharedPtr<FVisualLoggerTimeSliderController> GetTimeSliderController() { return TimeSliderController; }
 	UWorld* GetWorld(UObject* OptionalObject = nullptr);
 	class AActor* GetVisualLoggerHelperActor();
 	FVisualLoggerEvents& GetVisualLoggerEvents() { return VisualLoggerEvents; }
@@ -70,7 +70,7 @@ struct LOGVISUALIZER_API FLogVisualizer
 protected:
 	static TSharedPtr< struct FLogVisualizer > StaticInstance;
 	
-	TSharedPtr<FSequencerTimeSliderController> TimeSliderController;
+	TSharedPtr<FVisualLoggerTimeSliderController> TimeSliderController;
 	FVisualLoggerEvents VisualLoggerEvents;
 	TWeakPtr<IVisualLoggerInterface> VisualLoggerInterface;
 	TWeakPtr<class STimeline> CurrentTimeLine;
diff --git a/Engine/Source/Developer/LogVisualizer/Private/SSequencerSectionOverlay.h b/Engine/Source/Developer/LogVisualizer/Private/SSequencerSectionOverlay.h
--- a/Engine/Source/Developer/LogVisualizer/Private/SSequencerSectionOverlay.h
+++ b/Engine/Source/Developer/LogVisualizer/Private/SSequencerSectionOverlay.h
@@ -17,7 +17,7 @@ public:
 
 	SLATE_END_ARGS()
 
-	void Construct(const FArguments& InArgs, TSharedRef<class FSequencerTimeSliderController> InTimeSliderController)
+	void Construct(const FArguments& InArgs, TSharedRef<class FVisualLoggerTimeSliderController> InTimeSliderController)
 	{
 		bDisplayScrubPosition = InArgs._DisplayScrubPosition;
 		bDisplayTickLines = InArgs._DisplayTickLines;
@@ -34,7 +34,7 @@ private:
 
 private:
 	/** Controller for manipulating time */
-	TSharedPtr<class FSequencerTimeSliderController> TimeSliderController;
+	TSharedPtr<class FVisualLoggerTimeSliderController> TimeSliderController;
 	/** Whether or not to display the scrub position */
 	TAttribute<bool> bDisplayScrubPosition;
 	/** Whether or not to display tick lines */
diff --git a/Engine/Source/Developer/LogVisualizer/Private/STimeline.cpp b/Engine/Source/Developer/LogVisualizer/Private/STimeline.cpp
--- a/Engine/Source/Developer/LogVisualizer/Private/STimeline.cpp
+++ b/Engine/Source/Developer/LogVisualizer/Private/STimeline.cpp
@@ -183,7 +183,7 @@ void STimeline::AddEntry(const FVisualLogDevice::FVisualLogEntryItem& Entry)
 	}
 }
 
-void STimeline::Construct(const FArguments& InArgs, TSharedPtr<SVisualLoggerView> VisualLoggerView, TSharedPtr<FSequencerTimeSliderController> TimeSliderController, TSharedPtr<STimelinesContainer> InContainer, const FVisualLogDevice::FVisualLogEntryItem& Entry)
+void STimeline::Construct(const FArguments& InArgs, TSharedPtr<SVisualLoggerView> VisualLoggerView, TSharedPtr<FVisualLoggerTimeSliderController> TimeSliderController, TSharedPtr<STimelinesContainer> InContainer, const FVisualLogDevice::FVisualLogEntryItem& Entry)
 {
 	OnGetMenuContent = InArgs._OnGetMenuContent;
 
diff --git a/Engine/Source/Developer/LogVisualizer/Private/STimeline.h b/Engine/Source/Developer/LogVisualizer/Private/STimeline.h
--- a/Engine/Source/Developer/LogVisualizer/Private/STimeline.h
+++ b/Engine/Source/Developer/LogVisualizer/Private/STimeline.h
@@ -21,7 +21,7 @@ public:
 
 	/** Constructs this widget with InArgs */
 	virtual ~STimeline();
-	void Construct(const FArguments& InArgs, TSharedPtr<class SVisualLoggerView> VisualLoggerView, TSharedPtr<class FSequencerTimeSliderController> InTimeSliderController, TSharedPtr<class STimelinesContainer> InContainer, const FVisualLogDevice::FVisualLogEntryItem& Entry);
+	void Construct(const FArguments& InArgs, TSharedPtr<class SVisualLoggerView> VisualLoggerView, TSharedPtr<class FVisualLoggerTimeSliderController> InTimeSliderController, TSharedPtr<class STimelinesContainer> InContainer, const FVisualLogDevice::FVisualLogEntryItem& Entry);
 	bool IsSelected() const;
 	bool IsEntryHidden(const FVisualLogDevice::FVisualLogEntryItem&) const;
 	const FSlateBrush* GetBorder() const;
diff --git a/Engine/Source/Developer/LogVisualizer/Private/STimelineBar.cpp b/Engine/Source/Developer/LogVisualizer/Private/STimelineBar.cpp
--- a/Engine/Source/Developer/LogVisualizer/Private/STimelineBar.cpp
+++ b/Engine/Source/Developer/LogVisualizer/Private/STimelineBar.cpp
@@ -293,7 +293,7 @@ void STimelineBar::SnapScrubPosition(int32 NewItemIndex)
 	}
 }
 
-void STimelineBar::Construct(const FArguments& InArgs, TSharedPtr<FSequencerTimeSliderController> InTimeSliderController, TSharedPtr<STimeline> InTimelineOwner)
+void STimelineBar::Construct(const FArguments& InArgs, TSharedPtr<FVisualLoggerTimeSliderController> InTimeSliderController, TSharedPtr<STimeline> InTimelineOwner)
 {
 	TimeSliderController = InTimeSliderController;
 	TimelineOwner = InTimelineOwner;
diff --git a/Engine/Source/Developer/LogVisualizer/Private/STimelineBar.h b/Engine/Source/Developer/LogVisualizer/Private/STimelineBar.h
--- a/Engine/Source/Developer/LogVisualizer/Private/STimelineBar.h
+++ b/Engine/Source/Developer/LogVisualizer/Private/STimelineBar.h
@@ -22,7 +22,7 @@ public:
 
 	virtual bool SupportsKeyboardFocus() const { return true;  }
 
-	void Construct(const FArguments& InArgs, TSharedPtr<FSequencerTimeSliderController> InTimeSliderController, TSharedPtr<class STimeline> TimelineOwner);
+	void Construct(const FArguments& InArgs, TSharedPtr<FVisualLoggerTimeSliderController> InTimeSliderController, TSharedPtr<class STimeline> TimelineOwner);
 	virtual int32 OnPaint(const FPaintArgs& Args, const FGeometry& AllottedGeometry, const FSlateRect& MyClippingRect, FSlateWindowElementList& OutDrawElements, int32 LayerId, const FWidgetStyle& InWidgetStyle, bool bParentEnabled) const override;
 	FVector2D ComputeDesiredSize(float) const;
 	void SnapScrubPosition(float ScrubPosition);
@@ -35,7 +35,7 @@ public:
 	void GotoPreviousItem(int32 MoveDistance = 1);
 
 protected:
-	TSharedPtr<class FSequencerTimeSliderController> TimeSliderController;
+	TSharedPtr<class FVisualLoggerTimeSliderController> TimeSliderController;
 	TWeakPtr<class STimeline> TimelineOwner;
 	mutable int32 CurrentItemIndex;
 };
diff --git a/Engine/Source/Developer/LogVisualizer/Private/STimelinesContainer.cpp b/Engine/Source/Developer/LogVisualizer/Private/STimelinesContainer.cpp
--- a/Engine/Source/Developer/LogVisualizer/Private/STimelinesContainer.cpp
+++ b/Engine/Source/Developer/LogVisualizer/Private/STimelinesContainer.cpp
@@ -9,7 +9,7 @@
 
 #define LOCTEXT_NAMESPACE "STimelinesContainer"
 
-TSharedRef<SWidget> STimelinesContainer::MakeTimeline(TSharedPtr<class SVisualLoggerView> VisualLoggerView, TSharedPtr<class FSequencerTimeSliderController> TimeSliderController, const FVisualLogDevice::FVisualLogEntryItem& Entry)
+TSharedRef<SWidget> STimelinesContainer::MakeTimeline(TSharedPtr<class SVisualLoggerView> VisualLoggerView, TSharedPtr<class FVisualLoggerTimeSliderController> TimeSliderController, const FVisualLogDevice::FVisualLogEntryItem& Entry)
 {
 	TSharedPtr<STimeline> NewTimeline;
 
@@ -267,7 +267,7 @@ void STimelinesContainer::ResetData()
 	TimeSliderController->SetTimeRange(0, 5);
 }
 
-void STimelinesContainer::Construct(const FArguments& InArgs, TSharedRef<class SVisualLoggerView> InVisualLoggerView, TSharedRef<FSequencerTimeSliderController> InTimeSliderController)
+void STimelinesContainer::Construct(const FArguments& InArgs, TSharedRef<class SVisualLoggerView> InVisualLoggerView, TSharedRef<FVisualLoggerTimeSliderController> InTimeSliderController)
 {
 	TimeSliderController = InTimeSliderController;
 	VisualLoggerView = InVisualLoggerView;
diff --git a/Engine/Source/Developer/LogVisualizer/Private/STimelinesContainer.h b/Engine/Source/Developer/LogVisualizer/Private/STimelinesContainer.h
--- a/Engine/Source/Developer/LogVisualizer/Private/STimelinesContainer.h
+++ b/Engine/Source/Developer/LogVisualizer/Private/STimelinesContainer.h
@@ -18,8 +18,8 @@ public:
 	virtual FReply OnKeyDown(const FGeometry& MyGeometry, const FKeyEvent& InKeyEvent) override;
 
 	/** Constructs this widget with InArgs */
-	void Construct(const FArguments& InArgs, TSharedRef<class SVisualLoggerView>, TSharedRef<FSequencerTimeSliderController> TimeSliderController);
-	TSharedRef<SWidget> MakeTimeline(TSharedPtr<class SVisualLoggerView> VisualLoggerView, TSharedPtr<class FSequencerTimeSliderController> TimeSliderController, const FVisualLogDevice::FVisualLogEntryItem& Entry);
+	void Construct(const FArguments& InArgs, TSharedRef<class SVisualLoggerView>, TSharedRef<FVisualLoggerTimeSliderController> TimeSliderController);
+	TSharedRef<SWidget> MakeTimeline(TSharedPtr<class SVisualLoggerView> VisualLoggerView, TSharedPtr<class FVisualLoggerTimeSliderController> TimeSliderController, const FVisualLogDevice::FVisualLogEntryItem& Entry);
 	TSharedRef<SWidget> GetRightClickMenuContent();
 
 	void OnTimelineSelected(TSharedPtr<class STimelinesBar> Widget);
@@ -60,7 +60,7 @@ public:
 	bool IsNodeSelected(TSharedPtr<class STimeline> Node) const;
 
 protected:
-	TSharedPtr<class FSequencerTimeSliderController> TimeSliderController;
+	TSharedPtr<class FVisualLoggerTimeSliderController> TimeSliderController;
 	TSharedPtr<class SVisualLoggerView> VisualLoggerView;
 	TArray<TSharedPtr<class STimeline> > TimelineItems;
 	TArray< TSharedPtr<class STimeline> > SelectedNodes;
diff --git a/Engine/Source/Developer/LogVisualizer/Private/SVisualLoggerView.cpp b/Engine/Source/Developer/LogVisualizer/Private/SVisualLoggerView.cpp
--- a/Engine/Source/Developer/LogVisualizer/Private/SVisualLoggerView.cpp
+++ b/Engine/Source/Developer/LogVisualizer/Private/SVisualLoggerView.cpp
@@ -14,14 +14,14 @@
 class SInputCatcherOverlay : public SOverlay
 {
 public:
-	void Construct(const FArguments& InArgs, TSharedRef<class FSequencerTimeSliderController> InTimeSliderController)
+	void Construct(const FArguments& InArgs, TSharedRef<class FVisualLoggerTimeSliderController> InTimeSliderController)
 	{
 		SOverlay::Construct(InArgs);
 		TimeSliderController = InTimeSliderController;
 	}
 
 	/** Controller for manipulating time */
-	TSharedPtr<class FSequencerTimeSliderController> TimeSliderController;
+	TSharedPtr<class FVisualLoggerTimeSliderController> TimeSliderController;
 private:
 	/** SWidget Interface */
 	virtual FReply OnMouseButtonDown(const FGeometry& MyGeometry, const FPointerEvent& MouseEvent) override;
@@ -306,7 +306,7 @@ void SVisualLoggerView::OnSearchChanged(const FText& Filter)
 	TimelinesContainer->OnSearchChanged(Filter);
 }
 
-TSharedRef<SWidget> SVisualLoggerView::MakeSectionOverlay(TSharedRef<FSequencerTimeSliderController> TimeSliderController, const TAttribute< TRange<float> >& ViewRange, const TAttribute<float>& ScrubPosition, bool bTopOverlay)
+TSharedRef<SWidget> SVisualLoggerView::MakeSectionOverlay(TSharedRef<FVisualLoggerTimeSliderController> TimeSliderController, const TAttribute< TRange<float> >& ViewRange, const TAttribute<float>& ScrubPosition, bool bTopOverlay)
 {
 	return
 		SNew(SHorizontalBox)
diff --git a/Engine/Source/Developer/LogVisualizer/Private/SVisualLoggerView.h b/Engine/Source/Developer/LogVisualizer/Private/SVisualLoggerView.h
--- a/Engine/Source/Developer/LogVisualizer/Private/SVisualLoggerView.h
+++ b/Engine/Source/Developer/LogVisualizer/Private/SVisualLoggerView.h
@@ -25,7 +25,7 @@ public:
 	}
 	void SetAnimationOutlinerFillPercentage(float FillPercentage);
 
-	TSharedRef<SWidget> MakeSectionOverlay(TSharedRef<class FSequencerTimeSliderController> TimeSliderController, const TAttribute< TRange<float> >& ViewRange, const TAttribute<float>& ScrubPosition, bool bTopOverlay);
+	TSharedRef<SWidget> MakeSectionOverlay(TSharedRef<class FVisualLoggerTimeSliderController> TimeSliderController, const TAttribute< TRange<float> >& ViewRange, const TAttribute<float>& ScrubPosition, bool bTopOverlay);
 	void SetSearchString(FText SearchString);
 
 	void OnNewLogEntry(const FVisualLogDevice::FVisualLogEntryItem& Entry);
diff --git a/Engine/Source/Developer/LogVisualizer/Private/TimeSliderController.cpp b/Engine/Source/Developer/LogVisualizer/Private/TimeSliderController.cpp
--- a/Engine/Source/Developer/LogVisualizer/Private/TimeSliderController.cpp
+++ b/Engine/Source/Developer/LogVisualizer/Private/TimeSliderController.cpp
@@ -92,7 +92,7 @@ float DetermineOptimalSpacing( float InPixelsPerInput, uint32 MinTick, float Min
 }
 
 
-FSequencerTimeSliderController::FSequencerTimeSliderController(const FVisualLoggerTimeSliderArgs& InArgs)
+FVisualLoggerTimeSliderController::FVisualLoggerTimeSliderController(const FVisualLoggerTimeSliderArgs& InArgs)
 	: TimeSliderArgs( InArgs )
 	, DistanceDragged( 0.0f )
 	, bDraggingScrubber( false )
@@ -103,7 +103,7 @@ FSequencerTimeSliderController::FSequencerTimeSliderController(const FVisualLogg
 	CursorBackground = FEditorStyle::GetBrush("Sequencer.SectionArea.Background");
 }
 
-void FSequencerTimeSliderController::SetTimesliderArgs(const FVisualLoggerTimeSliderArgs& InArgs)
+void FVisualLoggerTimeSliderController::SetTimesliderArgs(const FVisualLoggerTimeSliderArgs& InArgs)
 {
 	TimeSliderArgs = InArgs;
 }
@@ -131,7 +131,7 @@ struct FDrawTickArgs
 	
 };
 
-void FSequencerTimeSliderController::DrawTicks( FSlateWindowElementList& OutDrawElements, const struct FScrubRangeToScreen& RangeToScreen, FDrawTickArgs& InArgs ) const
+void FVisualLoggerTimeSliderController::DrawTicks( FSlateWindowElementList& OutDrawElements, const struct FScrubRangeToScreen& RangeToScreen, FDrawTickArgs& InArgs ) const
 {
 	const float Spacing = DetermineOptimalSpacing( RangeToScreen.PixelsPerInput, ScrubConstants::MinPixelsPerDisplayTick, ScrubConstants::MinDisplayTickSpacing );
 
@@ -229,7 +229,7 @@ void FSequencerTimeSliderController::DrawTicks( FSlateWindowElementList& OutDraw
 }
 
 
-int32 FSequencerTimeSliderController::OnPaintTimeSlider( bool bMirrorLabels, const FGeometry& AllottedGeometry, const FSlateRect& MyClippingRect, FSlateWindowElementList& OutDrawElements, int32 LayerId, const FWidgetStyle& InWidgetStyle, bool bParentEnabled ) const
+int32 FVisualLoggerTimeSliderController::OnPaintTimeSlider( bool bMirrorLabels, const FGeometry& AllottedGeometry, const FSlateRect& MyClippingRect, FSlateWindowElementList& OutDrawElements, int32 LayerId, const FWidgetStyle& InWidgetStyle, bool bParentEnabled ) const
 {
 	const bool bEnabled = bParentEnabled;
 	const ESlateDrawEffect::Type DrawEffects = bEnabled ? ESlateDrawEffect::None : ESlateDrawEffect::DisabledEffect;
@@ -310,7 +310,7 @@ int32 FSequencerTimeSliderController::OnPaintTimeSlider( bool bMirrorLabels, con
 	return LayerId;
 }
 
-FReply FSequencerTimeSliderController::OnMouseButtonDown( TSharedRef<SWidget> WidgetOwner, const FGeometry& MyGeometry, const FPointerEvent& MouseEvent )
+FReply FVisualLoggerTimeSliderController::OnMouseButtonDown( TSharedRef<SWidget> WidgetOwner, const FGeometry& MyGeometry, const FPointerEvent& MouseEvent )
 {
 	bool bHandleLeftMouseButton = MouseEvent.GetEffectingButton() == EKeys::LeftMouseButton;
 	bool bHandleRightMouseButton = MouseEvent.GetEffectingButton() == EKeys::RightMouseButton && TimeSliderArgs.AllowZoom;
@@ -335,7 +335,7 @@ FReply FSequencerTimeSliderController::OnMouseButtonDown( TSharedRef<SWidget> Wi
 	return FReply::Unhandled();
 }
 
-FReply FSequencerTimeSliderController::OnMouseButtonUp( TSharedRef<SWidget> WidgetOwner, const FGeometry& MyGeometry, const FPointerEvent& MouseEvent )
+FReply FVisualLoggerTimeSliderController::OnMouseButtonUp( TSharedRef<SWidget> WidgetOwner, const FGeometry& MyGeometry, const FPointerEvent& MouseEvent )
 {
 	bool bHandleLeftMouseButton = MouseEvent.GetEffectingButton() == EKeys::LeftMouseButton && WidgetOwner->HasMouseCapture();
 	bool bHandleRightMouseButton = MouseEvent.GetEffectingButton() == EKeys::RightMouseButton && WidgetOwner->HasMouseCapture() && TimeSliderArgs.AllowZoom ;
@@ -376,7 +376,7 @@ FReply FSequencerTimeSliderController::OnMouseButtonUp( TSharedRef<SWidget> Widg
 	return FReply::Unhandled();
 }
 
-FReply FSequencerTimeSliderController::OnMouseMove( TSharedRef<SWidget> WidgetOwner, const FGeometry& MyGeometry, const FPointerEvent& MouseEvent )
+FReply FVisualLoggerTimeSliderController::OnMouseMove( TSharedRef<SWidget> WidgetOwner, const FGeometry& MyGeometry, const FPointerEvent& MouseEvent )
 {
 	if ( WidgetOwner->HasMouseCapture() )
 	{
@@ -472,7 +472,7 @@ FReply FSequencerTimeSliderController::OnMouseMove( TSharedRef<SWidget> WidgetOw
 	return FReply::Unhandled();
 }
 
-void FSequencerTimeSliderController::CommitScrubPosition( float NewValue, bool bIsScrubbing )
+void FVisualLoggerTimeSliderController::CommitScrubPosition( float NewValue, bool bIsScrubbing )
 {
 	// Manage the scrub position ourselves if its not bound to a delegate
 	if ( !TimeSliderArgs.ScrubPosition.IsBound() )
@@ -494,13 +494,13 @@ void FSequencerTimeSliderController::CommitScrubPosition( float NewValue, bool b
 	TimeSliderArgs.OnScrubPositionChanged.ExecuteIfBound( NewValue, bIsScrubbing );
 }
 
-void FSequencerTimeSliderController::SetExternalScrollbar(TSharedRef<SScrollBar> InScrollbar) 
+void FVisualLoggerTimeSliderController::SetExternalScrollbar(TSharedRef<SScrollBar> InScrollbar) 
 { 
 	Scrollbar = InScrollbar;
-	Scrollbar->SetOnUserScrolled(FOnUserScrolled::CreateRaw(this, &FSequencerTimeSliderController::HorizontalScrollBar_OnUserScrolled)); 
+	Scrollbar->SetOnUserScrolled(FOnUserScrolled::CreateRaw(this, &FVisualLoggerTimeSliderController::HorizontalScrollBar_OnUserScrolled)); 
 };
 
-void FSequencerTimeSliderController::HorizontalScrollBar_OnUserScrolled(float ScrollOffset)
+void FVisualLoggerTimeSliderController::HorizontalScrollBar_OnUserScrolled(float ScrollOffset)
 {
 	if (!TimeSliderArgs.ViewRange.IsBound())
 	{
@@ -526,7 +526,7 @@ void FSequencerTimeSliderController::HorizontalScrollBar_OnUserScrolled(float Sc
 	}
 }
 
-void FSequencerTimeSliderController::SetTimeRange(float NewViewOutputMin, float NewViewOutputMax)
+void FVisualLoggerTimeSliderController::SetTimeRange(float NewViewOutputMin, float NewViewOutputMax)
 {
 	TimeSliderArgs.ViewRange.Set(TRange<float>(NewViewOutputMin, NewViewOutputMax));
 
@@ -538,7 +538,7 @@ void FSequencerTimeSliderController::SetTimeRange(float NewViewOutputMin, float
 	Scrollbar->SetState(InOffsetFraction, InThumbSizeFraction);
 }
 
-void FSequencerTimeSliderController::SetClampRange(float MinValue, float MaxValue)
+void FVisualLoggerTimeSliderController::SetClampRange(float MinValue, float MaxValue)
 {
 	TRange<float> LocalViewRange = TimeSliderArgs.ViewRange.Get();
 	float LocalClampMin = TimeSliderArgs.ClampMin.Get().GetValue();
@@ -556,7 +556,7 @@ void FSequencerTimeSliderController::SetClampRange(float MinValue, float MaxValu
 	SetTimeRange(ZoomDelta >= 1 ? MinValue : LocalViewRangeMin, ZoomDelta >= 1 ? MaxValue : LocalViewRangeMax);
 }
 
-FReply FSequencerTimeSliderController::OnMouseWheel( TSharedRef<SWidget> WidgetOwner, const FGeometry& MyGeometry, const FPointerEvent& MouseEvent )
+FReply FVisualLoggerTimeSliderController::OnMouseWheel( TSharedRef<SWidget> WidgetOwner, const FGeometry& MyGeometry, const FPointerEvent& MouseEvent )
 {
 	FReply ReturnValue = FReply::Unhandled();;
 
@@ -618,7 +618,7 @@ FReply FSequencerTimeSliderController::OnMouseWheel( TSharedRef<SWidget> WidgetO
 	return ReturnValue;
 }
 
-int32 FSequencerTimeSliderController::OnPaintSectionView( const FGeometry& AllottedGeometry, const FSlateRect& MyClippingRect, FSlateWindowElementList& OutDrawElements, int32 LayerId, bool bEnabled, bool bDisplayTickLines, bool bDisplayScrubPosition  ) const
+int32 FVisualLoggerTimeSliderController::OnPaintSectionView( const FGeometry& AllottedGeometry, const FSlateRect& MyClippingRect, FSlateWindowElementList& OutDrawElements, int32 LayerId, bool bEnabled, bool bDisplayTickLines, bool bDisplayScrubPosition  ) const
 {
 	const ESlateDrawEffect::Type DrawEffects = bEnabled ? ESlateDrawEffect::None : ESlateDrawEffect::DisabledEffect;
 
diff --git a/Engine/Source/Developer/LogVisualizer/Private/TimeSliderController.h b/Engine/Source/Developer/LogVisualizer/Private/TimeSliderController.h
--- a/Engine/Source/Developer/LogVisualizer/Private/TimeSliderController.h
+++ b/Engine/Source/Developer/LogVisualizer/Private/TimeSliderController.h
@@ -20,10 +20,10 @@ struct FVisualLoggerTimeSliderArgs : FTimeSliderArgs
  * A time slider controller for sequencer
  * Draws and manages time data for a Sequencer
  */
-class FSequencerTimeSliderController : public ITimeSliderController
+class FVisualLoggerTimeSliderController : public ITimeSliderController
 {
 public:
-	FSequencerTimeSliderController(const FVisualLoggerTimeSliderArgs& InArgs);
+	FVisualLoggerTimeSliderController(const FVisualLoggerTimeSliderArgs& InArgs);
 	void SetTimesliderArgs(const FVisualLoggerTimeSliderArgs& InArgs);
 
 	/** ITimeSliderController Interface */
