Fixing one more name conflict.

- Missed in PR #778.

diff --git a/Engine/Source/Developer/LogVisualizer/Private/SSequencerSectionOverlay.cpp b/Engine/Source/Developer/LogVisualizer/Private/SSequencerSectionOverlay.cpp
--- a/Engine/Source/Developer/LogVisualizer/Private/SSequencerSectionOverlay.cpp
+++ b/Engine/Source/Developer/LogVisualizer/Private/SSequencerSectionOverlay.cpp
@@ -4,29 +4,29 @@
 #include "SSequencerSectionOverlay.h"
 #include "TimeSliderController.h"
 
-int32 SSequencerSectionOverlay::OnPaint( const FPaintArgs& Args, const FGeometry& AllottedGeometry, const FSlateRect& MyClippingRect, FSlateWindowElementList& OutDrawElements, int32 LayerId, const FWidgetStyle& InWidgetStyle, bool bParentEnabled ) const
+int32 SVisualLoggerSectionOverlay::OnPaint( const FPaintArgs& Args, const FGeometry& AllottedGeometry, const FSlateRect& MyClippingRect, FSlateWindowElementList& OutDrawElements, int32 LayerId, const FWidgetStyle& InWidgetStyle, bool bParentEnabled ) const
 {
 	TimeSliderController->OnPaintSectionView( AllottedGeometry, MyClippingRect, OutDrawElements, LayerId, ShouldBeEnabled( bParentEnabled ), bDisplayTickLines.Get(), bDisplayScrubPosition.Get() );
 
 	return SCompoundWidget::OnPaint( Args, AllottedGeometry, MyClippingRect, OutDrawElements, LayerId, InWidgetStyle, bParentEnabled  );
 }
 
-FReply SSequencerSectionOverlay::OnMouseButtonDown( const FGeometry& MyGeometry, const FPointerEvent& MouseEvent )
+FReply SVisualLoggerSectionOverlay::OnMouseButtonDown( const FGeometry& MyGeometry, const FPointerEvent& MouseEvent )
 {
 	return TimeSliderController->OnMouseButtonDown(SharedThis(this), MyGeometry, MouseEvent);
 }
 
-FReply SSequencerSectionOverlay::OnMouseButtonUp( const FGeometry& MyGeometry, const FPointerEvent& MouseEvent )
+FReply SVisualLoggerSectionOverlay::OnMouseButtonUp( const FGeometry& MyGeometry, const FPointerEvent& MouseEvent )
 {
 	return TimeSliderController->OnMouseButtonUp( SharedThis(this),  MyGeometry, MouseEvent );
 }
 
-FReply SSequencerSectionOverlay::OnMouseMove( const FGeometry& MyGeometry, const FPointerEvent& MouseEvent )
+FReply SVisualLoggerSectionOverlay::OnMouseMove( const FGeometry& MyGeometry, const FPointerEvent& MouseEvent )
 {
 	return TimeSliderController->OnMouseMove( SharedThis(this), MyGeometry, MouseEvent );
 }
 
-FReply SSequencerSectionOverlay::OnMouseWheel( const FGeometry& MyGeometry, const FPointerEvent& MouseEvent )
+FReply SVisualLoggerSectionOverlay::OnMouseWheel( const FGeometry& MyGeometry, const FPointerEvent& MouseEvent )
 {
 	if (MouseEvent.IsLeftShiftDown() || MouseEvent.IsLeftControlDown())
 	{
diff --git a/Engine/Source/Developer/LogVisualizer/Private/SSequencerSectionOverlay.h b/Engine/Source/Developer/LogVisualizer/Private/SSequencerSectionOverlay.h
--- a/Engine/Source/Developer/LogVisualizer/Private/SSequencerSectionOverlay.h
+++ b/Engine/Source/Developer/LogVisualizer/Private/SSequencerSectionOverlay.h
@@ -5,10 +5,10 @@
 /**
  * An overlay that displays global information in the section area
  */
-class SSequencerSectionOverlay : public SCompoundWidget
+class SVisualLoggerSectionOverlay : public SCompoundWidget
 {
 public:
-	SLATE_BEGIN_ARGS( SSequencerSectionOverlay )
+	SLATE_BEGIN_ARGS( SVisualLoggerSectionOverlay )
 		: _DisplayTickLines( true )
 		, _DisplayScrubPosition( false )
 	{}
diff --git a/Engine/Source/Developer/LogVisualizer/Private/SVisualLoggerView.cpp b/Engine/Source/Developer/LogVisualizer/Private/SVisualLoggerView.cpp
--- a/Engine/Source/Developer/LogVisualizer/Private/SVisualLoggerView.cpp
+++ b/Engine/Source/Developer/LogVisualizer/Private/SVisualLoggerView.cpp
@@ -320,7 +320,7 @@ TSharedRef<SWidget> SVisualLoggerView::MakeSectionOverlay(TSharedRef<FVisualLogg
 		+ SHorizontalBox::Slot()
 		.FillWidth(1.0f)
 		[
-			SNew(SSequencerSectionOverlay, TimeSliderController)
+			SNew(SVisualLoggerSectionOverlay, TimeSliderController)
 			.DisplayScrubPosition(bTopOverlay)
 			.DisplayTickLines(!bTopOverlay)
 		];
