Integrated change to fix issue "Pause does not stop the Visual Logger from recording data while PIEing with dedicated server" from Dev to Main.
 - Added option to LogVisualizer's settings, to select between regular pause or PlayersOnly pause - to pause gamplay on level if needed.

diff --git a/Engine/Source/Developer/LogVisualizer/Private/LogVisualizerSettings.cpp b/Engine/Source/Developer/LogVisualizer/Private/LogVisualizerSettings.cpp
--- a/Engine/Source/Developer/LogVisualizer/Private/LogVisualizerSettings.cpp
+++ b/Engine/Source/Developer/LogVisualizer/Private/LogVisualizerSettings.cpp
@@ -18,6 +18,7 @@ ULogVisualizerSettings::ULogVisualizerSettings(const FObjectInitializer& ObjectI
 	GraphsBackgroundColor = FColor(0, 0, 0, 70);
 	bResetDataWithNewSession = false;
 	bDrawExtremesOnGraphs = false;
+	bUsePlayersOnlyForPause = true;
 }
 
 #if WITH_EDITOR
diff --git a/Engine/Source/Developer/LogVisualizer/Private/SVisualLogger.cpp b/Engine/Source/Developer/LogVisualizer/Private/SVisualLogger.cpp
--- a/Engine/Source/Developer/LogVisualizer/Private/SVisualLogger.cpp
+++ b/Engine/Source/Developer/LogVisualizer/Private/SVisualLogger.cpp
@@ -107,6 +107,7 @@ void SVisualLogger::FVisualLoggerDevice::Serialize(const class UObject* LogOwner
 
 SVisualLogger::SVisualLogger() : SCompoundWidget(), CommandList(MakeShareable(new FUICommandList))
 { 
+	bPausedLogger = false;
 	InternalDevice = MakeShareable(new FVisualLoggerDevice(this));
 	FVisualLogger::Get().AddDevice(InternalDevice.Get());
 }
@@ -168,6 +169,8 @@ SVisualLogger::~SVisualLogger()
 
 void SVisualLogger::Construct(const FArguments& InArgs, const TSharedRef<SDockTab>& ConstructUnderMajorTab, const TSharedPtr<SWindow>& ConstructUnderWindow)
 {
+	bPausedLogger = false;
+
 	FLogVisualizer::Get().SetCurrentVisualizer(SharedThis(this));
 	//////////////////////////////////////////////////////////////////////////
 	// Visual Logger Events
@@ -509,17 +512,24 @@ bool SVisualLogger::HandleStopRecordingCommandIsVisible() const
 
 bool SVisualLogger::HandlePauseCommandCanExecute() const
 {
-	UWorld* World = FLogVisualizer::Get().GetWorld();
-	return FVisualLogger::Get().IsRecording() && World && !World->bPlayersOnly && !World->bPlayersOnlyPending;
+	return !bPausedLogger;
 }
 
 void SVisualLogger::HandlePauseCommandExecute()
 {
-	UWorld* World = FLogVisualizer::Get().GetWorld();
-	if (World != NULL)
+	if (ULogVisualizerSettings::StaticClass()->GetDefaultObject<ULogVisualizerSettings>()->bUsePlayersOnlyForPause)
 	{
-		World->bPlayersOnlyPending = true;
+		const TIndirectArray<FWorldContext>& WorldContexts = GEngine->GetWorldContexts();
+		for (const FWorldContext& Context : WorldContexts)
+		{
+			if (Context.World() != nullptr)
+			{
+				Context.World()->bPlayersOnlyPending = true;
+			}
+		}
 	}
+
+	bPausedLogger = true;
 }
 
 bool SVisualLogger::HandlePauseCommandIsVisible() const
@@ -529,18 +539,31 @@ bool SVisualLogger::HandlePauseCommandIsVisible() const
 
 bool SVisualLogger::HandleResumeCommandCanExecute() const
 {
-	UWorld* World = FLogVisualizer::Get().GetWorld();
-	return FVisualLogger::Get().IsRecording() && World && (World->bPlayersOnly || World->bPlayersOnlyPending);
+	return bPausedLogger;
 }
 
 void SVisualLogger::HandleResumeCommandExecute()
 {
-	UWorld* World = FLogVisualizer::Get().GetWorld();
-	if (World != NULL)
+	if (ULogVisualizerSettings::StaticClass()->GetDefaultObject<ULogVisualizerSettings>()->bUsePlayersOnlyForPause)
 	{
-		World->bPlayersOnly = false;
-		World->bPlayersOnlyPending = false;
+		const TIndirectArray<FWorldContext>& WorldContexts = GEngine->GetWorldContexts();
+		for (const FWorldContext& Context : WorldContexts)
+		{
+			if (Context.World() != nullptr)
+			{
+				Context.World()->bPlayersOnly = false;
+				Context.World()->bPlayersOnlyPending = false;
+			}
+		}
+	}
+
+	bPausedLogger = false;
+	for (const auto& CurrentEntry : OnPauseCacheForEntries)
+	{
+		OnNewLogEntry(CurrentEntry);
 	}
+	OnPauseCacheForEntries.Reset();
+
 }
 
 bool SVisualLogger::HandleResumeCommandIsVisible() const
@@ -756,8 +779,15 @@ void SVisualLogger::OnNewWorld(UWorld* NewWorld)
 
 void SVisualLogger::OnNewLogEntry(const FVisualLogDevice::FVisualLogEntryItem& Entry)
 {
-	CollectNewCategories(Entry);
-	MainView->OnNewLogEntry(Entry);
+	if (!bPausedLogger)
+	{
+		CollectNewCategories(Entry);
+		MainView->OnNewLogEntry(Entry);
+	}
+	else
+	{
+		OnPauseCacheForEntries.Add(Entry);
+	}
 }
 
 void SVisualLogger::CollectNewCategories(const FVisualLogDevice::FVisualLogEntryItem& Item)
diff --git a/Engine/Source/Developer/LogVisualizer/Private/SVisualLogger.h b/Engine/Source/Developer/LogVisualizer/Private/SVisualLogger.h
--- a/Engine/Source/Developer/LogVisualizer/Private/SVisualLogger.h
+++ b/Engine/Source/Developer/LogVisualizer/Private/SVisualLogger.h
@@ -125,5 +125,8 @@ protected:
 	mutable TSharedPtr<SVisualLoggerLogsList> LogsList;
 	mutable TSharedPtr<SVisualLoggerStatusView> StatusView;
 
+	bool bPausedLogger;
+	TArray<FVisualLogDevice::FVisualLogEntryItem> OnPauseCacheForEntries;
+
 	FDelegateHandle DrawOnCanvasDelegateHandle;
 };
diff --git a/Engine/Source/Developer/LogVisualizer/Private/SVisualLoggerToolbar.cpp b/Engine/Source/Developer/LogVisualizer/Private/SVisualLoggerToolbar.cpp
--- a/Engine/Source/Developer/LogVisualizer/Private/SVisualLoggerToolbar.cpp
+++ b/Engine/Source/Developer/LogVisualizer/Private/SVisualLoggerToolbar.cpp
@@ -21,10 +21,10 @@ TSharedRef<SWidget> SVisualLoggerToolbar::MakeToolbar(const TSharedRef<FUIComman
 
 	ToolBarBuilder.BeginSection("Debugger");
 	{
-		ToolBarBuilder.AddToolBarButton(FVisualLoggerCommands::Get().StartRecording, NAME_None, LOCTEXT("StartDebugger", "Start"), LOCTEXT("StartDebuggerTooltip", "Starts recording and collecting visual logs"), FSlateIcon(FLogVisualizerStyle::Get().GetStyleSetName(), TEXT("Toolbar.Record")));
-		ToolBarBuilder.AddToolBarButton(FVisualLoggerCommands::Get().StopRecording, NAME_None, LOCTEXT("StopDebugger", "Stop"), TAttribute<FText>(), FSlateIcon(FLogVisualizerStyle::Get().GetStyleSetName(), TEXT("Toolbar.Stop")));
-		ToolBarBuilder.AddToolBarButton(FVisualLoggerCommands::Get().Pause, NAME_None, LOCTEXT("PauseGame", "Pause"), TAttribute<FText>(), FSlateIcon(FLogVisualizerStyle::Get().GetStyleSetName(), TEXT("Toolbar.Pause")));
-		ToolBarBuilder.AddToolBarButton(FVisualLoggerCommands::Get().Resume, NAME_None, LOCTEXT("ResumeGame", "Resume"), TAttribute<FText>(), FSlateIcon(FLogVisualizerStyle::Get().GetStyleSetName(), TEXT("Toolbar.Resume")));
+		ToolBarBuilder.AddToolBarButton(FVisualLoggerCommands::Get().StartRecording, NAME_None, LOCTEXT("StartLogger", "Start"), LOCTEXT("StartDebuggerTooltip", "Starts recording and collecting visual logs"), FSlateIcon(FLogVisualizerStyle::Get().GetStyleSetName(), TEXT("Toolbar.Record")));
+		ToolBarBuilder.AddToolBarButton(FVisualLoggerCommands::Get().StopRecording, NAME_None, LOCTEXT("StopLogger", "Stop"), TAttribute<FText>(), FSlateIcon(FLogVisualizerStyle::Get().GetStyleSetName(), TEXT("Toolbar.Stop")));
+		ToolBarBuilder.AddToolBarButton(FVisualLoggerCommands::Get().Pause, NAME_None, LOCTEXT("PauseLogger", "Pause"), TAttribute<FText>(), FSlateIcon(FLogVisualizerStyle::Get().GetStyleSetName(), TEXT("Toolbar.Pause")));
+		ToolBarBuilder.AddToolBarButton(FVisualLoggerCommands::Get().Resume, NAME_None, LOCTEXT("ResumeLogger", "Resume"), TAttribute<FText>(), FSlateIcon(FLogVisualizerStyle::Get().GetStyleSetName(), TEXT("Toolbar.Resume")));
 
 		ToolBarBuilder.AddSeparator();
 		ToolBarBuilder.AddToolBarButton(FVisualLoggerCommands::Get().LoadFromVLog, NAME_None, LOCTEXT("Load", "Load"), TAttribute<FText>(), FSlateIcon(FLogVisualizerStyle::Get().GetStyleSetName(), TEXT("Toolbar.Load")));
diff --git a/Engine/Source/Developer/LogVisualizer/Public/LogVisualizerSettings.h b/Engine/Source/Developer/LogVisualizer/Public/LogVisualizerSettings.h
--- a/Engine/Source/Developer/LogVisualizer/Public/LogVisualizerSettings.h
+++ b/Engine/Source/Developer/LogVisualizer/Public/LogVisualizerSettings.h
@@ -80,6 +80,10 @@ public:
 	UPROPERTY(EditAnywhere, config, Category = "VisualLogger")
 	bool bDrawExtremesOnGraphs;
 
+	/** Whether to use PlayersOnly during Pause or not */
+	UPROPERTY(EditAnywhere, config, Category = "VisualLogger")
+	bool bUsePlayersOnlyForPause;
+
 	UPROPERTY(config)
 	FFiltersPreset PresistentPresets;
 
