Fix for UE-7497 "No Graphs are displayed in Visual Logger after enabling them" issue. Toggle graph button is hidden for logs without any graph/histogram data now.

diff --git a/Engine/Source/Developer/LogVisualizer/Private/SVisualLogger.cpp b/Engine/Source/Developer/LogVisualizer/Private/SVisualLogger.cpp
--- a/Engine/Source/Developer/LogVisualizer/Private/SVisualLogger.cpp
+++ b/Engine/Source/Developer/LogVisualizer/Private/SVisualLogger.cpp
@@ -108,6 +108,7 @@ void SVisualLogger::FVisualLoggerDevice::Serialize(const class UObject* LogOwner
 SVisualLogger::SVisualLogger() : SCompoundWidget(), CommandList(MakeShareable(new FUICommandList))
 { 
 	bPausedLogger = false;
+	bGotHistogramData = false;
 	InternalDevice = MakeShareable(new FVisualLoggerDevice(this));
 	FVisualLogger::Get().AddDevice(InternalDevice.Get());
 }
@@ -170,6 +171,7 @@ SVisualLogger::~SVisualLogger()
 void SVisualLogger::Construct(const FArguments& InArgs, const TSharedRef<SDockTab>& ConstructUnderMajorTab, const TSharedPtr<SWindow>& ConstructUnderWindow)
 {
 	bPausedLogger = false;
+	bGotHistogramData = false;
 
 	FLogVisualizer::Get().SetCurrentVisualizer(SharedThis(this));
 	//////////////////////////////////////////////////////////////////////////
@@ -204,9 +206,9 @@ void SVisualLogger::Construct(const FArguments& InArgs, const TSharedRef<SDockTa
 		FIsActionButtonVisible::CreateRaw(this, &SVisualLogger::HandleCameraCommandCanExecute));
 	ActionList.MapAction(Commands.ToggleGraphs,
 		FExecuteAction::CreateLambda([](){bool& bEnableGraphsVisualization = ULogVisualizerSessionSettings::StaticClass()->GetDefaultObject<ULogVisualizerSessionSettings>()->bEnableGraphsVisualization; bEnableGraphsVisualization = !bEnableGraphsVisualization; }),
-		FCanExecuteAction(),
+		FCanExecuteAction::CreateLambda([this]()->bool{return bGotHistogramData; }),
 		FIsActionChecked::CreateLambda([]()->bool{return ULogVisualizerSessionSettings::StaticClass()->GetDefaultObject<ULogVisualizerSessionSettings>()->bEnableGraphsVisualization; }),
-		FIsActionButtonVisible());
+		FIsActionButtonVisible::CreateLambda([this]()->bool{return bGotHistogramData; }));
 	ActionList.MapAction(Commands.ResetData, 
 		FExecuteAction::CreateRaw(this, &SVisualLogger::ResetData),
 		FCanExecuteAction::CreateRaw(this, &SVisualLogger::HandleSaveCommandCanExecute),
@@ -766,6 +768,8 @@ void SVisualLogger::ResetData()
 		HelperActor->ObjectSelectionChanged(NULL);
 	}
 
+	bGotHistogramData = false;
+	OnPauseCacheForEntries.Reset();
 }
 
 void SVisualLogger::OnNewWorld(UWorld* NewWorld)
@@ -783,6 +787,7 @@ void SVisualLogger::OnNewLogEntry(const FVisualLogDevice::FVisualLogEntryItem& E
 	{
 		CollectNewCategories(Entry);
 		MainView->OnNewLogEntry(Entry);
+		bGotHistogramData = bGotHistogramData || Entry.Entry.HistogramSamples.Num() > 0;
 	}
 	else
 	{
diff --git a/Engine/Source/Developer/LogVisualizer/Private/SVisualLogger.h b/Engine/Source/Developer/LogVisualizer/Private/SVisualLogger.h
--- a/Engine/Source/Developer/LogVisualizer/Private/SVisualLogger.h
+++ b/Engine/Source/Developer/LogVisualizer/Private/SVisualLogger.h
@@ -128,5 +128,7 @@ protected:
 	bool bPausedLogger;
 	TArray<FVisualLogDevice::FVisualLogEntryItem> OnPauseCacheForEntries;
 
+	bool bGotHistogramData;
+
 	FDelegateHandle DrawOnCanvasDelegateHandle;
 };
