Fixed issue which could give us inaccessible vlogs in LogVisualizer with persistent filters.
Made small refactor in code, to fix how LogVisualizer works with filters.

diff --git a/Engine/Source/Developer/LogVisualizer/Private/LogVisualizer.h b/Engine/Source/Developer/LogVisualizer/Private/LogVisualizer.h
--- a/Engine/Source/Developer/LogVisualizer/Private/LogVisualizer.h
+++ b/Engine/Source/Developer/LogVisualizer/Private/LogVisualizer.h
@@ -34,14 +34,6 @@ struct FVisualLoggerEvents
 	FOnObjectSelectionChanged OnObjectSelectionChanged;
 };
 
-struct IVisualLoggerInterface
-{
-	virtual ~IVisualLoggerInterface() {}
-	virtual bool HasValidCategories(TArray<FVisualLoggerCategoryVerbosityPair> Categories) = 0;
-	virtual bool IsValidCategory(const FString& InCategoryName, TEnumAsByte<ELogVerbosity::Type> Verbosity = ELogVerbosity::All) = 0;
-	virtual bool IsValidCategory(const FString& InGraphName, const FString& InDataName, TEnumAsByte<ELogVerbosity::Type> Verbosity = ELogVerbosity::All) = 0;
-};
-
 class FVisualLoggerTimeSliderController;
 struct LOGVISUALIZER_API FLogVisualizer
 {
@@ -57,8 +49,6 @@ struct LOGVISUALIZER_API FLogVisualizer
 	class AActor* GetVisualLoggerHelperActor();
 	FVisualLoggerEvents& GetVisualLoggerEvents() { return VisualLoggerEvents; }
 
-	TSharedPtr<IVisualLoggerInterface> GetVisualLoggerInterface() { return VisualLoggerInterface.Pin(); }
-	void SetVisualLoggerInterface(TSharedPtr<IVisualLoggerInterface> InVisualLoggerInterface) { VisualLoggerInterface = InVisualLoggerInterface; }
 	void SetCurrentVisualizer(TSharedPtr<class SVisualLogger> Visualizer) { CurrentVisualizer = Visualizer; }
 
 	void OnObjectSelectionChanged(TSharedPtr<class STimeline> TimeLine) { CurrentTimeLine = TimeLine; }
@@ -72,7 +62,6 @@ protected:
 	
 	TSharedPtr<FVisualLoggerTimeSliderController> TimeSliderController;
 	FVisualLoggerEvents VisualLoggerEvents;
-	TWeakPtr<IVisualLoggerInterface> VisualLoggerInterface;
 	TWeakPtr<class STimeline> CurrentTimeLine;
 	TWeakPtr<class SVisualLogger> CurrentVisualizer;
 };
diff --git a/Engine/Source/Developer/LogVisualizer/Private/LogVisualizerSettings.cpp b/Engine/Source/Developer/LogVisualizer/Private/LogVisualizerSettings.cpp
--- a/Engine/Source/Developer/LogVisualizer/Private/LogVisualizerSettings.cpp
+++ b/Engine/Source/Developer/LogVisualizer/Private/LogVisualizerSettings.cpp
@@ -10,6 +10,8 @@
 #include "UnrealEdMisc.h"
 #endif // WITH_EDITOR
 
+FCategoryFiltersManager FCategoryFiltersManager::StaticManager;
+
 ULogVisualizerSettings::ULogVisualizerSettings(const FObjectInitializer& ObjectInitializer) : Super(ObjectInitializer)
 {
 	TrivialLogsThreshold = 1;
@@ -35,3 +37,159 @@ void ULogVisualizerSettings::PostEditChangeProperty(struct FPropertyChangedEvent
 	SettingChangedEvent.Broadcast(Name);
 }
 #endif
+
+bool FCategoryFiltersManager::MatchCategoryFilters(FString String, ELogVerbosity::Type Verbosity)
+{
+	const ULogVisualizerSettings* Settings = ULogVisualizerSettings::StaticClass()->GetDefaultObject<ULogVisualizerSettings>();
+	bool bFoundFilter = false;
+	for (const FCategoryFilter& Filter : Settings->CurrentFilters.Categories)
+	{
+		bFoundFilter = Filter.CategoryName == String;
+		if (bFoundFilter)
+		{
+			const bool bFineWithSearchString = Settings->bSearchInsideLogs == true ||
+				(Settings->bSearchInsideLogs == false && (Settings->CurrentFilters.SearchBoxFilter.Len() == 0 || Filter.CategoryName.Find(Settings->CurrentFilters.SearchBoxFilter) != INDEX_NONE));
+			return bFineWithSearchString && Verbosity <= Filter.LogVerbosity && Filter.Enabled;
+		}
+	}
+
+	return bFoundFilter;
+}
+
+bool FCategoryFiltersManager::MatchObjectName(FString String)
+{
+	const ULogVisualizerSettings* Settings = ULogVisualizerSettings::StaticClass()->GetDefaultObject<ULogVisualizerSettings>();
+	return (Settings->CurrentFilters.SelectedClasses.Num() == 0 || Settings->CurrentFilters.SelectedClasses.Find(String) != INDEX_NONE);
+}
+
+bool FCategoryFiltersManager::MatchSearchString(FString String)
+{
+	const ULogVisualizerSettings* Settings = ULogVisualizerSettings::StaticClass()->GetDefaultObject<ULogVisualizerSettings>();
+	return Settings->CurrentFilters.SearchBoxFilter == String;
+}
+
+void FCategoryFiltersManager::SetSearchString(FString InString) 
+{ 
+	ULogVisualizerSettings::StaticClass()->GetDefaultObject<ULogVisualizerSettings>()->CurrentFilters.SearchBoxFilter = InString; 
+}
+
+FString FCategoryFiltersManager::GetSearchString() 
+{ 
+	return ULogVisualizerSettings::StaticClass()->GetDefaultObject<ULogVisualizerSettings>()->CurrentFilters.SearchBoxFilter; 
+}
+
+void FCategoryFiltersManager::SetObjectFilterString(FString InFilterString) 
+{ 
+	ULogVisualizerSettings::StaticClass()->GetDefaultObject<ULogVisualizerSettings>()->CurrentFilters.ObjectNameFilter = InFilterString; 
+}
+
+FString FCategoryFiltersManager::GetObjectFilterString()
+{ 
+	return ULogVisualizerSettings::StaticClass()->GetDefaultObject<ULogVisualizerSettings>()->CurrentFilters.ObjectNameFilter; 
+}
+
+void FCategoryFiltersManager::AddCategory(FString InName, ELogVerbosity::Type InVerbosity)
+{
+	ULogVisualizerSettings* Settings = ULogVisualizerSettings::StaticClass()->GetDefaultObject<ULogVisualizerSettings>();
+	FCategoryFilter Filter;
+	Filter.CategoryName = InName;
+	Filter.LogVerbosity = InVerbosity;
+	Filter.Enabled = true;
+	Settings->CurrentFilters.Categories.Add(Filter);
+}
+
+void FCategoryFiltersManager::RemoveCategory(FString InName)
+{
+	ULogVisualizerSettings* Settings = ULogVisualizerSettings::StaticClass()->GetDefaultObject<ULogVisualizerSettings>();
+	for (int32 Index = 0; Index < Settings->CurrentFilters.Categories.Num(); ++Index)
+	{
+		const FCategoryFilter& Filter = Settings->CurrentFilters.Categories[Index];
+		if (Filter.CategoryName == InName)
+		{
+			Settings->CurrentFilters.Categories.RemoveAt(Index);
+			break;
+		}
+	}
+}
+
+bool FCategoryFiltersManager::IsValidCategory(FString InName)
+{
+	ULogVisualizerSettings* Settings = ULogVisualizerSettings::StaticClass()->GetDefaultObject<ULogVisualizerSettings>();
+	for (const FCategoryFilter& Filter : Settings->CurrentFilters.Categories)
+	{
+		if (Filter.CategoryName == InName)
+		{
+			return true;
+		}
+	}
+	return Settings->CurrentFilters.Categories.Num() == 0;
+}
+
+FCategoryFilter& FCategoryFiltersManager::GetCategory(FString InName)
+{
+	ULogVisualizerSettings* Settings = ULogVisualizerSettings::StaticClass()->GetDefaultObject<ULogVisualizerSettings>();
+	for (FCategoryFilter& Filter : Settings->CurrentFilters.Categories)
+	{
+		if (Filter.CategoryName == InName)
+		{
+			return Filter;
+		}
+	}
+
+	static FCategoryFilter NoCategory;
+	return NoCategory;
+}
+
+void FCategoryFiltersManager::SelectObject(FString ObjectName)
+{
+	ULogVisualizerSettings* Settings = ULogVisualizerSettings::StaticClass()->GetDefaultObject<ULogVisualizerSettings>();
+	Settings->CurrentFilters.SelectedClasses.AddUnique(ObjectName);
+}
+
+void FCategoryFiltersManager::RemoveObjectFromSelection(FString ObjectName)
+{
+	ULogVisualizerSettings* Settings = ULogVisualizerSettings::StaticClass()->GetDefaultObject<ULogVisualizerSettings>();
+	Settings->CurrentFilters.SelectedClasses.Remove(ObjectName);
+}
+
+const TArray<FString>& FCategoryFiltersManager::GetSelectedObjects()
+{
+	return ULogVisualizerSettings::StaticClass()->GetDefaultObject<ULogVisualizerSettings>()->CurrentFilters.SelectedClasses;
+}
+
+void FCategoryFiltersManager::SavePresistentData()
+{
+	ULogVisualizerSettings* Settings = ULogVisualizerSettings::StaticClass()->GetDefaultObject<ULogVisualizerSettings>();
+	if (Settings->bPresistentFilters)
+	{
+		Settings->PresistentFilters = Settings->CurrentFilters;
+	}
+	else
+	{
+		Settings->PresistentFilters = FVisualLoggerFilters();
+	}
+	Settings->SaveConfig();
+
+}
+
+void FCategoryFiltersManager::ClearPresistentData()
+{
+	ULogVisualizerSettings* Settings = ULogVisualizerSettings::StaticClass()->GetDefaultObject<ULogVisualizerSettings>();
+	if (Settings->bPresistentFilters)
+	{
+		Settings->PresistentFilters = FVisualLoggerFilters();
+	}
+}
+
+void FCategoryFiltersManager::LoadPresistentData()
+{
+	ULogVisualizerSettings* Settings = ULogVisualizerSettings::StaticClass()->GetDefaultObject<ULogVisualizerSettings>();
+	if (Settings->bPresistentFilters)
+	{
+		Settings->CurrentFilters = Settings->PresistentFilters;
+	}
+	else
+	{
+		Settings->CurrentFilters = FVisualLoggerFilters();
+	}
+}
\ No newline at end of file
diff --git a/Engine/Source/Developer/LogVisualizer/Private/SFilterWidget.cpp b/Engine/Source/Developer/LogVisualizer/Private/SFilterWidget.cpp
--- a/Engine/Source/Developer/LogVisualizer/Private/SFilterWidget.cpp
+++ b/Engine/Source/Developer/LogVisualizer/Private/SFilterWidget.cpp
@@ -7,7 +7,6 @@
 /** Constructs this widget with InArgs */
 void SFilterWidget::Construct(const FArguments& InArgs)
 {
-	bEnabled = true;
 	OnFilterChanged = InArgs._OnFilterChanged;
 	OnRequestRemove = InArgs._OnRequestRemove;
 	OnRequestEnableOnly = InArgs._OnRequestEnableOnly;
@@ -15,7 +14,6 @@ void SFilterWidget::Construct(const FArguments& InArgs)
 	OnRequestRemoveAll = InArgs._OnRequestRemoveAll;
 	FilterName = InArgs._FilterName;
 	ColorCategory = InArgs._ColorCategory;
-	Verbosity = ELogVerbosity::All;
 	BorderBackgroundColor = FLinearColor(0.2f, 0.2f, 0.2f, 0.2f);
 
 	// Get the tooltip and color of the type represented by this filter
@@ -53,8 +51,9 @@ void SFilterWidget::Construct(const FArguments& InArgs)
 FText SFilterWidget::GetCaptionString() const
 {
 	FString CaptionString;
-	const FString VerbosityStr = FOutputDevice::VerbosityToString(Verbosity);
-	if (Verbosity != ELogVerbosity::VeryVerbose)
+	FCategoryFilter& Filter = FCategoryFiltersManager::Get().GetCategory(GetFilterNameAsString());
+	const FString VerbosityStr = FOutputDevice::VerbosityToString((ELogVerbosity::Type)Filter.LogVerbosity);
+	if (Filter.LogVerbosity != ELogVerbosity::VeryVerbose)
 	{
 		CaptionString = FString::Printf(TEXT("%s [%s]"), *GetFilterNameAsString().Replace(TEXT("Log"), TEXT(""), ESearchCase::CaseSensitive), *VerbosityStr.Mid(0, 1));
 	}
@@ -67,9 +66,12 @@ FText SFilterWidget::GetCaptionString() const
 
 FText SFilterWidget::GetTooltipString() const
 {
+	FCategoryFilter& Filter = FCategoryFiltersManager::Get().GetCategory(GetFilterNameAsString());
+	const FString VerbosityStr = FOutputDevice::VerbosityToString((ELogVerbosity::Type)Filter.LogVerbosity);
+
 	return FText::FromString(
 		IsEnabled() ?
-		FString::Printf(TEXT("Enabled '%s' category for '%s' verbosity and lower\nRight click to change verbosity"), *GetFilterNameAsString(), FOutputDevice::VerbosityToString(Verbosity))
+		FString::Printf(TEXT("Enabled '%s' category for '%s' verbosity and lower\nRight click to change verbosity"), *GetFilterNameAsString(), *VerbosityStr)
 		:
 		FString::Printf(TEXT("Disabled '%s' category"), *GetFilterNameAsString())
 		);
@@ -100,4 +102,34 @@ TSharedRef<SWidget> SFilterWidget::GetRightClickMenuContent()
 
 	return MenuBuilder.MakeWidget();
 }
+
+bool SFilterWidget::IsEnabled() const 
+{ 
+	FCategoryFilter& Filter = FCategoryFiltersManager::Get().GetCategory(GetFilterNameAsString());
+	return Filter.Enabled;
+}
+
+void SFilterWidget::SetEnabled(bool InEnabled)
+{
+	FCategoryFilter& Filter = FCategoryFiltersManager::Get().GetCategory(GetFilterNameAsString());
+	if (InEnabled != Filter.Enabled)
+	{
+		Filter.Enabled = InEnabled;
+		OnFilterChanged.ExecuteIfBound();
+	}
+}
+
+void SFilterWidget::FilterToggled(ECheckBoxState NewState)
+{
+	FCategoryFilter& Filter = FCategoryFiltersManager::Get().GetCategory(GetFilterNameAsString());
+	SetEnabled(NewState == ECheckBoxState::Checked);
+}
+
+void SFilterWidget::SetVerbosityFilter(int32 SelectedVerbosityIndex)
+{
+	FCategoryFilter& Filter = FCategoryFiltersManager::Get().GetCategory(GetFilterNameAsString());
+	Filter.LogVerbosity = (ELogVerbosity::Type)SelectedVerbosityIndex;
+	OnFilterChanged.ExecuteIfBound();
+}
+
 #undef LOCTEXT_NAMESPACE
diff --git a/Engine/Source/Developer/LogVisualizer/Private/SFilterWidget.h b/Engine/Source/Developer/LogVisualizer/Private/SFilterWidget.h
--- a/Engine/Source/Developer/LogVisualizer/Private/SFilterWidget.h
+++ b/Engine/Source/Developer/LogVisualizer/Private/SFilterWidget.h
@@ -38,21 +38,10 @@ public:
 	FText GetTooltipString() const;
 
 	/** Sets whether or not this filter is applied to the combined filter */
-	void SetEnabled(bool InEnabled)
-	{
-		if (InEnabled != bEnabled)
-		{
-			bEnabled = InEnabled;
-			OnFilterChanged.ExecuteIfBound();
-		}
-	}
+	void SetEnabled(bool InEnabled);
 
 	/** Returns true if this filter contributes to the combined filter */
-	bool IsEnabled() const { return bEnabled; }
-
-	void SetVerbosity(ELogVerbosity::Type InVerbosity) { Verbosity = InVerbosity; }
-
-	ELogVerbosity::Type GetVerbosity() const { return Verbosity; }
+	bool IsEnabled() const;
 
 	FORCEINLINE FName GetFilterName() const { return FilterName; }
 
@@ -72,19 +61,14 @@ public:
 
 private:
 	/** Handler for when the filter checkbox is clicked */
-	void FilterToggled(ECheckBoxState NewState)
-	{
-		bEnabled = NewState == ECheckBoxState::Checked;
-		OnFilterChanged.ExecuteIfBound();
-	}
+	void FilterToggled(ECheckBoxState NewState);
 
 	/** Handler for when the filter checkbox is double clicked */
 	FReply FilterDoubleClicked()
 	{
 		// Disable all other filters and enable this one.
 		OnRequestDisableAll.ExecuteIfBound();
-		bEnabled = true;
-		OnFilterChanged.ExecuteIfBound();
+		SetEnabled(true);
 
 		return FReply::Handled();
 	}
@@ -100,11 +84,7 @@ private:
 	TSharedRef<SWidget> GetRightClickMenuContent();
 
 	/** Removes this filter from the filter list */
-	void SetVerbosityFilter(int32 SelectedVerbosityIndex)
-	{
-		Verbosity = (ELogVerbosity::Type)SelectedVerbosityIndex;
-		OnFilterChanged.ExecuteIfBound();
-	}
+	void SetVerbosityFilter(int32 SelectedVerbosityIndex);
 
 	void RemoveFilter()
 	{
@@ -126,7 +106,7 @@ private:
 	void RemoveAllFilters() { OnRequestRemoveAll.ExecuteIfBound(); }
 
 	/** Handler to determine the "checked" state of the filter checkbox */
-	ECheckBoxState IsChecked() const { return bEnabled ? ECheckBoxState::Checked : ECheckBoxState::Unchecked; }
+	ECheckBoxState IsChecked() const { return IsEnabled() ? ECheckBoxState::Checked : ECheckBoxState::Unchecked; }
 
 	/** Handler to determine the color of the checkbox when it is checked */
 	FSlateColor GetFilterForegroundColor() const { return IsChecked() == ECheckBoxState::Checked ? FilterColor : FLinearColor::White; }
@@ -162,15 +142,10 @@ private:
 	/** Invoked when a request to remove all filters originated from within this filter */
 	FOnSimpleRequest OnRequestRemoveAll;
 
-	/** true when this filter should be applied to the search */
-	bool bEnabled;
-
 	FName FilterName;
 
 	FLinearColor ColorCategory;
 
-	ELogVerbosity::Type Verbosity;
-
 	/** The button to toggle the filter on or off */
 	TSharedPtr<SFilterCheckBox> ToggleButtonPtr;
 
diff --git a/Engine/Source/Developer/LogVisualizer/Private/STimeline.cpp b/Engine/Source/Developer/LogVisualizer/Private/STimeline.cpp
--- a/Engine/Source/Developer/LogVisualizer/Private/STimeline.cpp
+++ b/Engine/Source/Developer/LogVisualizer/Private/STimeline.cpp
@@ -76,7 +76,7 @@ void STimeline::UpdateVisibility()
 {
 	ULogVisualizerSettings* Settings = ULogVisualizerSettings::StaticClass()->GetDefaultObject<ULogVisualizerSettings>();
 	const bool bJustIgnore = Settings->bIgnoreTrivialLogs && Entries.Num() <= Settings->TrivialLogsThreshold;
-	const bool bVisibleByOwnerClasses = Settings->CurrentPresets.SelectedClasses.Num() == 0 || Settings->CurrentPresets.SelectedClasses.Find(OwnerClassName.ToString()) != INDEX_NONE;
+	const bool bVisibleByOwnerClasses = FCategoryFiltersManager::Get().MatchObjectName(OwnerClassName.ToString());
 	const bool bIsCollapsed = bJustIgnore || HiddenEntries.Num() == Entries.Num() || (SearchFilter.Len() > 0 && Name.ToString().Find(SearchFilter) == INDEX_NONE);
 
 	SetVisibility(bIsCollapsed || !bVisibleByOwnerClasses ? EVisibility::Collapsed : EVisibility::Visible);
@@ -90,7 +90,7 @@ void STimeline::UpdateVisibilityForItems()
 {
 	HiddenEntries.Reset();
 	ULogVisualizerSettings* Settings = ULogVisualizerSettings::StaticClass()->GetDefaultObject<ULogVisualizerSettings>();
-	const FString QuickSearchStrng = Settings->CurrentPresets.DataFilter;
+	const FString QuickSearchStrng = FCategoryFiltersManager::Get().GetSearchString(); 
 
 	for (auto& CurrentEntry : Entries)
 	{
@@ -100,7 +100,6 @@ void STimeline::UpdateVisibilityForItems()
 
 void STimeline::OnFiltersSearchChanged(const FText& Filter)
 {
-	//QuickSearchStrng = Filter.ToString();
 	OnFiltersChanged();
 }
 
@@ -130,7 +129,11 @@ void STimeline::UpdateVisibilityForEntry(FVisualLogDevice::FVisualLogEntryItem&
 {
 	TArray<FVisualLoggerCategoryVerbosityPair> OutCategories;
 	FVisualLoggerHelpers::GetCategories(CurrentEntry.Entry, OutCategories);
-	bool bHasValidCategories = FLogVisualizer::Get().GetVisualLoggerInterface()->HasValidCategories(OutCategories);
+	bool bHasValidCategories = false;
+	for (FVisualLoggerCategoryVerbosityPair& Categoryair : OutCategories)
+	{
+		bHasValidCategories = bHasValidCategories || FCategoryFiltersManager::Get().MatchCategoryFilters(Categoryair.CategoryName.ToString(), Categoryair.Verbosity);
+	}
 
 	if (bSearchInsideLogs && bHasValidCategories && SearchString.Len() > 0)
 	{
@@ -174,7 +177,8 @@ void STimeline::AddEntry(const FVisualLogDevice::FVisualLogEntryItem& Entry)
 	Entries.Add(Entry); 
 
 	ULogVisualizerSettings* Settings = ULogVisualizerSettings::StaticClass()->GetDefaultObject<ULogVisualizerSettings>();
-	UpdateVisibilityForEntry(Entries[Entries.Num() - 1], Settings->CurrentPresets.DataFilter, Settings->bSearchInsideLogs);
+	
+	UpdateVisibilityForEntry(Entries[Entries.Num() - 1], FCategoryFiltersManager::Get().GetSearchString(), Settings->bSearchInsideLogs);
 	UpdateVisibility();
 
 	if (Settings->bStickToRecentData && IsSelected())
diff --git a/Engine/Source/Developer/LogVisualizer/Private/SVisualLogger.cpp b/Engine/Source/Developer/LogVisualizer/Private/SVisualLogger.cpp
--- a/Engine/Source/Developer/LogVisualizer/Private/SVisualLogger.cpp
+++ b/Engine/Source/Developer/LogVisualizer/Private/SVisualLogger.cpp
@@ -34,48 +34,10 @@ static const FName StatusViewTabId("StatusView");
 namespace LogVisualizer
 {
 	static const FString LogFileDescription = LOCTEXT("FileTypeDescription", "Visual Log File").ToString();
-	static const FString LoadFileTypes = FString::Printf(TEXT("%s (*.vlog;*.%s)|*.vlog;*.%s"), *LogFileDescription, VISLOG_FILENAME_EXT, VISLOG_FILENAME_EXT);
+	static const FString LoadFileTypes = FString::Printf(TEXT("%s (*.bvlog;*.%s)|*.bvlog;*.%s"), *LogFileDescription, VISLOG_FILENAME_EXT, VISLOG_FILENAME_EXT);
 	static const FString SaveFileTypes = FString::Printf(TEXT("%s (*.%s)|*.%s"), *LogFileDescription, VISLOG_FILENAME_EXT, VISLOG_FILENAME_EXT);
 }
 
-struct FVisualLoggerInterface : public IVisualLoggerInterface
-{
-	FVisualLoggerInterface(TSharedPtr<SVisualLogger> InVisualLogger)
-		: VisualLogger(InVisualLogger)
-	{
-	}
-
-	virtual ~FVisualLoggerInterface() 
-	{
-
-	}
-
-	virtual bool HasValidCategories(TArray<FVisualLoggerCategoryVerbosityPair> Categories) override
-	{
-		for (const auto& CurrentCategory : Categories)
-		{
-			if (VisualLogger.Pin()->GetVisualLoggerFilters()->IsFilterEnabled(CurrentCategory.CategoryName.ToString(), CurrentCategory.Verbosity))
-			{
-				return true;
-			}
-		}
-		return false;
-	}
-
-	bool IsValidCategory(const FString& InCategoryName, TEnumAsByte<ELogVerbosity::Type> Verbosity)
-	{
-		return VisualLogger.Pin()->GetVisualLoggerFilters()->IsFilterEnabled(InCategoryName, Verbosity);
-	}
-
-	bool IsValidCategory(const FString& InGraphName, const FString& InDataName, TEnumAsByte<ELogVerbosity::Type> Verbosity = ELogVerbosity::All)
-	{
-		return VisualLogger.Pin()->GetVisualLoggerFilters()->IsFilterEnabled(InGraphName, InDataName, Verbosity);
-	}
-
-protected:
-	TWeakPtr<SVisualLogger> VisualLogger;
-};
-
 SVisualLogger::FVisualLoggerDevice::FVisualLoggerDevice(SVisualLogger* InOwner)
 	:Owner(InOwner)
 {
@@ -117,17 +79,7 @@ SVisualLogger::~SVisualLogger()
 {
 	UWorld* World = NULL;
 #if WITH_EDITOR
-	ULogVisualizerSettings* Settings = ULogVisualizerSettings::StaticClass()->GetDefaultObject<ULogVisualizerSettings>();
-	if (Settings->bPresistentFilters)
-	{
-		Settings->PresistentPresets = Settings->CurrentPresets;
-	}
-	else
-	{
-		Settings->PresistentPresets = Settings->CurrentPresets = FFiltersPreset();
-	}
-	Settings->SaveConfig();
-
+	FCategoryFiltersManager::Get().SavePresistentData();
 
 	UEditorEngine *EEngine = Cast<UEditorEngine>(GEngine);
 	if (GIsEditor && EEngine != NULL)
@@ -161,8 +113,6 @@ SVisualLogger::~SVisualLogger()
 
 	FVisualLogger::Get().RemoveDevice(InternalDevice.Get());
 	InternalDevice.Reset();
-	
-	VisualLoggerInterface.Reset();
 }
 
 /* SMessagingDebugger interface
@@ -180,18 +130,13 @@ void SVisualLogger::Construct(const FArguments& InArgs, const TSharedRef<SDockTa
 	FLogVisualizer::Get().GetVisualLoggerEvents().OnFiltersChanged = FOnFiltersChanged::CreateRaw(this, &SVisualLogger::OnFiltersChanged);
 	FLogVisualizer::Get().GetVisualLoggerEvents().OnObjectSelectionChanged = FOnObjectSelectionChanged::CreateRaw(this, &SVisualLogger::OnObjectSelectionChanged);
 
-	//////////////////////////////////////////////////////////////////////////
-	// Visual Logger Interface
-	VisualLoggerInterface = MakeShareable(new FVisualLoggerInterface(SharedThis(this)));
-	FLogVisualizer::Get().SetVisualLoggerInterface(VisualLoggerInterface);
-
 	//////////////////////////////////////////////////////////////////////////
 	// Command Action Lists
 	const FVisualLoggerCommands& Commands = FVisualLoggerCommands::Get();
 	FUICommandList& ActionList = *CommandList;
 
 	ULogVisualizerSettings* Settings = ULogVisualizerSettings::StaticClass()->GetDefaultObject<ULogVisualizerSettings>();
-	Settings->CurrentPresets = Settings->bPresistentFilters ? Settings->PresistentPresets : FFiltersPreset();
+	FCategoryFiltersManager::Get().LoadPresistentData();
 
 	ActionList.MapAction(Commands.StartRecording, FExecuteAction::CreateRaw(this, &SVisualLogger::HandleStartRecordingCommandExecute), FCanExecuteAction::CreateRaw(this, &SVisualLogger::HandleStartRecordingCommandCanExecute), FIsActionChecked(), FIsActionButtonVisible::CreateRaw(this, &SVisualLogger::HandleStartRecordingCommandIsVisible));
 	ActionList.MapAction(Commands.StopRecording, FExecuteAction::CreateRaw(this, &SVisualLogger::HandleStopRecordingCommandExecute), FCanExecuteAction::CreateRaw(this, &SVisualLogger::HandleStopRecordingCommandCanExecute), FIsActionChecked(), FIsActionButtonVisible::CreateRaw(this, &SVisualLogger::HandleStopRecordingCommandIsVisible));
@@ -406,6 +351,7 @@ void SVisualLogger::FillLoadPresetMenu(FMenuBuilder& MenuBuilder)
 	ULogVisualizerSettings* Settings = ULogVisualizerSettings::StaticClass()->GetDefaultObject<ULogVisualizerSettings>();
 	
 	MenuBuilder.BeginSection("FilterPresets", LOCTEXT("FilterPresets", "Presets"));
+#if 0 //disabled for now, we don't use presets yet
 	for (auto& CurrentPreset : Settings->FilterPresets)
 	{
 		MenuBuilder.AddMenuEntry(
@@ -415,12 +361,13 @@ void SVisualLogger::FillLoadPresetMenu(FMenuBuilder& MenuBuilder)
 			))
 		);
 	}
+#endif
 	MenuBuilder.EndSection();
 }
 
 void SVisualLogger::SetFiltersPreset(const struct FFiltersPreset& Preset)
 {
-	OnFiltersSearchChanged(FText::FromString(Preset.DataFilter));
+
 }
 
 void SVisualLogger::FillWindowMenu(FMenuBuilder& MenuBuilder, const TSharedPtr<FTabManager> TabManager)
@@ -848,8 +795,7 @@ void SVisualLogger::OnObjectSelectionChanged(TSharedPtr<class STimeline> TimeLin
 
 void SVisualLogger::OnFiltersSearchChanged(const FText& Filter)
 {
-	ULogVisualizerSettings::StaticClass()->GetDefaultObject<ULogVisualizerSettings>()->CurrentPresets.DataFilter = Filter.ToString();
-	ULogVisualizerSettings::StaticClass()->GetDefaultObject<ULogVisualizerSettings>()->SaveConfig();
+	FCategoryFiltersManager::Get().SetSearchString(Filter.ToString());
 
 	if (VisualLoggerFilters.IsValid())
 	{
diff --git a/Engine/Source/Developer/LogVisualizer/Private/SVisualLogger.h b/Engine/Source/Developer/LogVisualizer/Private/SVisualLogger.h
--- a/Engine/Source/Developer/LogVisualizer/Private/SVisualLogger.h
+++ b/Engine/Source/Developer/LogVisualizer/Private/SVisualLogger.h
@@ -118,7 +118,6 @@ protected:
 
 	TWeakObjectPtr<class AVisualLoggerCameraController> CameraController;
 	TSharedPtr<struct FVisualLoggerCanvasRenderer> VisualLoggerCanvasRenderer;
-	TSharedPtr<IVisualLoggerInterface> VisualLoggerInterface;
 
 	mutable TSharedPtr<SVisualLoggerFilters> VisualLoggerFilters;
 	mutable TSharedPtr<SVisualLoggerView> MainView;
diff --git a/Engine/Source/Developer/LogVisualizer/Private/SVisualLoggerFilters.cpp b/Engine/Source/Developer/LogVisualizer/Private/SVisualLoggerFilters.cpp
--- a/Engine/Source/Developer/LogVisualizer/Private/SVisualLoggerFilters.cpp
+++ b/Engine/Source/Developer/LogVisualizer/Private/SVisualLoggerFilters.cpp
@@ -68,12 +68,15 @@ bool SVisualLoggerFilters::GraphSubmenuVisibility(const FName MenuName)
 		return true;
 	}
 
-	const TArray<FSimpleGraphFilter> DataNames = *GraphFilters.Find(MenuName);
-	for (const auto& CurrentData : DataNames)
+	if (GraphFilters.Find(MenuName) != NULL)
 	{
-		if (CurrentData.Name.ToString().Find(GraphsFilter) != INDEX_NONE)
+		const TArray<FString> DataNames = *GraphFilters.Find(MenuName);
+		for (const auto& CurrentData : DataNames)
 		{
-			return true;
+			if (CurrentData.Find(GraphsFilter) != INDEX_NONE)
+			{
+				return true;
+			}
 		}
 	}
 
@@ -141,7 +144,8 @@ void SVisualLoggerFilters::GraphFilterCategoryClicked(FName MenuCategory)
 		auto &Filters = GraphFilters[MenuCategory];
 		for (int32 Index = 0; Index < Filters.Num(); ++Index)
 		{
-			Filters[Index].bEnabled = bNewSet;
+			FString GraphFilterName = MenuCategory.ToString() + TEXT("$") + Filters[Index];
+			FCategoryFiltersManager::Get().GetCategory(GraphFilterName).Enabled = bNewSet;
 			bChanged = true;
 		}
 
@@ -154,19 +158,18 @@ void SVisualLoggerFilters::GraphFilterCategoryClicked(FName MenuCategory)
 
 bool SVisualLoggerFilters::IsGraphFilterCategoryInUse(FName MenuCategory) const
 {
+	bool bInUse = false;
 	if (GraphFilters.Contains(MenuCategory))
 	{
 		auto &Filters = GraphFilters[MenuCategory];
 		for (int32 Index = 0; Index < Filters.Num(); ++Index)
 		{
-			if (Filters[Index].bEnabled)
-			{
-				return true;
-			}
+			FString GraphFilterName = MenuCategory.ToString() + TEXT("$") + Filters[Index];
+			bInUse |= FCategoryFiltersManager::Get().GetCategory(GraphFilterName).Enabled;
 		}
 	}
 
-	return false;
+	return bInUse;
 }
 
 
@@ -175,15 +178,16 @@ void SVisualLoggerFilters::CreateFiltersMenuCategoryForGraph(FMenuBuilder& MenuB
 	auto FiltersFromGraph = GraphFilters[MenuCategory];
 	for (auto Iter = FiltersFromGraph.CreateIterator(); Iter; ++Iter)
 	{
-		const FText& LabelText = FText::FromString(Iter->Name.ToString());
+		FName Name = **Iter;
+		const FText& LabelText = FText::FromString(Name.ToString());
 		MenuBuilder.AddMenuEntry(
 			LabelText,
 			FText::Format(LOCTEXT("FilterByTooltipPrefix", "Filter by {0}"), LabelText),
 			FSlateIcon(),
 			FUIAction(
-			FExecuteAction::CreateSP(this, &SVisualLoggerFilters::FilterByTypeClicked, MenuCategory, Iter->Name),
+			FExecuteAction::CreateSP(this, &SVisualLoggerFilters::FilterByTypeClicked, MenuCategory, Name),
 			FCanExecuteAction(),
-			FIsActionChecked::CreateSP(this, &SVisualLoggerFilters::IsAssetTypeActionsInUse, MenuCategory, Iter->Name),
+			FIsActionChecked::CreateSP(this, &SVisualLoggerFilters::IsAssetTypeActionsInUse, MenuCategory, Name),
 			FIsActionButtonVisible::CreateLambda([this, LabelText]()->bool{return this->GraphsFilter.Len() == 0 || LabelText.ToString().Find(this->GraphsFilter) != INDEX_NONE; })),
 			NAME_None,
 			EUserInterfaceActionType::ToggleButton
@@ -199,9 +203,11 @@ void SVisualLoggerFilters::FilterByTypeClicked(FName InGraphName, FName InDataNa
 		auto &Filters = GraphFilters[InGraphName];
 		for (int32 Index = 0; Index < Filters.Num(); ++Index)
 		{
-			if (Filters[Index].Name == InDataName)
+			if (Filters[Index] == InDataName.ToString())
 			{
-				Filters[Index].bEnabled = !Filters[Index].bEnabled;
+				FString GraphFilterName = InGraphName.ToString() + TEXT("$") + Filters[Index];
+				FCategoryFilter& Filter = FCategoryFiltersManager::Get().GetCategory(GraphFilterName);
+				Filter.Enabled = !Filter.Enabled;
 				bChanged = true;
 			}
 		}
@@ -220,9 +226,10 @@ bool SVisualLoggerFilters::IsAssetTypeActionsInUse(FName InGraphName, FName InDa
 		auto &Filters = GraphFilters[InGraphName];
 		for (int32 Index = 0; Index < Filters.Num(); ++Index)
 		{
-			if (Filters[Index].Name == InDataName)
+			if (Filters[Index] == InDataName.ToString())
 			{
-				return Filters[Index].bEnabled;
+				FString GraphFilterName = InGraphName.ToString() + TEXT("$") + Filters[Index];
+				return FCategoryFiltersManager::Get().GetCategory(GraphFilterName).Enabled;
 			}
 		}
 	}
@@ -291,11 +298,14 @@ void SVisualLoggerFilters::AddFilter(const FString& InFilterName)
 		//.OnRequestRemoveAll(this, &SLogFilterList::RemoveAllFilters);
 		;
 
-	ULogVisualizerSettings* Settings = ULogVisualizerSettings::StaticClass()->GetDefaultObject<ULogVisualizerSettings>();
-	TArray<FString> CurrentFilters = Settings->CurrentPresets.CategoryFilters;
-	NewFilter->SetEnabled(CurrentFilters.Num() == 0 || CurrentFilters.Find(InFilterName) != INDEX_NONE);
+	bool bIsSet = true;
+
+	if (FCategoryFiltersManager::Get().IsValidCategory(InFilterName))
+	{
+		bIsSet = FCategoryFiltersManager::Get().GetCategory(InFilterName).Enabled;
+	}
+	FCategoryFiltersManager::Get().AddCategory(InFilterName, ELogVerbosity::All);
 	Filters.Add(NewFilter);
-	Settings->CurrentPresets.CategoryFilters = CurrentFilters;
 
 	FilterBox->AddSlot()
 		.Padding(2, 2)
@@ -326,53 +336,16 @@ void SVisualLoggerFilters::OnFiltersChanged()
 			EnabledFilters.AddUnique(CurrentFilter->GetFilterName().ToString());
 		}
 	}
-	ULogVisualizerSettings* Settings = ULogVisualizerSettings::StaticClass()->GetDefaultObject<ULogVisualizerSettings>();
-	Settings->CurrentPresets.CategoryFilters = EnabledFilters;
-	Settings->SaveConfig();
 
 	FLogVisualizer::Get().GetVisualLoggerEvents().OnFiltersChanged.ExecuteIfBound();
 }
 
 void SVisualLoggerFilters::AddFilter(const FString& GraphName, const FString& DataName)
 {
-	TArray<FSimpleGraphFilter>& DataNames = GraphFilters.FindOrAdd(*GraphName);
-	if (DataNames.Find(FSimpleGraphFilter(*DataName)) == INDEX_NONE)
-	{
-		DataNames.Add(FSimpleGraphFilter(*DataName));
-	}
-}
-
-bool SVisualLoggerFilters::IsFilterEnabled(const FString& InFilterName, TEnumAsByte<ELogVerbosity::Type> Verbosity)
-{
-	const ULogVisualizerSettings* Settings = ULogVisualizerSettings::StaticClass()->GetDefaultObject<ULogVisualizerSettings>();
-	FString FiltersSearchString = Settings->CurrentPresets.DataFilter;
-	const FName FilterName(*InFilterName);
-	for (int32 Index = 0; Index < Filters.Num(); ++Index)
-	{
-		const SFilterWidget& Filter = Filters[Index].Get();
-		if (Filter.GetFilterName() == FilterName)
-		{
-			return Filter.IsEnabled() && Filter.GetVerbosity() >= Verbosity.GetValue() && 
-				(Settings->bSearchInsideLogs == true || (Settings->bSearchInsideLogs == false && (FiltersSearchString.Len() == 0 || Filter.GetFilterName().ToString().Find(FiltersSearchString) != INDEX_NONE)));
-		}
-	}
-
-	return true;
-}
-
-bool SVisualLoggerFilters::IsFilterEnabled(const FString& InGraphName, const FString& InDataName, TEnumAsByte<ELogVerbosity::Type> Verbosity)
-{
-	if (GraphFilters.Contains(*InGraphName))
-	{
-		auto Filters = GraphFilters[*InGraphName];
-		int32 Index = Filters.Find(FSimpleGraphFilter(*InDataName));
-		if (Index != INDEX_NONE)
-		{
-			return Filters[Index].bEnabled && (GraphsFilter.Len() == 0 || Filters[Index].Name.ToString().Find(GraphsFilter) != INDEX_NONE);
-		}
-	}
+	GraphFilters.FindOrAdd(*GraphName).AddUnique(DataName);
 
-	return true;
+	FString GraphFilterName = GraphName + TEXT("$") + DataName;
+	FCategoryFiltersManager::Get().AddCategory(GraphFilterName, ELogVerbosity::All);
 }
 
 void SVisualLoggerFilters::OnFiltersSearchChanged(const FText& Filter)
diff --git a/Engine/Source/Developer/LogVisualizer/Private/SVisualLoggerFilters.h b/Engine/Source/Developer/LogVisualizer/Private/SVisualLoggerFilters.h
--- a/Engine/Source/Developer/LogVisualizer/Private/SVisualLoggerFilters.h
+++ b/Engine/Source/Developer/LogVisualizer/Private/SVisualLoggerFilters.h
@@ -5,21 +5,6 @@
 /**
 * A list of filters currently applied to an asset view.
 */
-struct FSimpleGraphFilter
-{
-	FSimpleGraphFilter(FName InName) : Name(InName), bEnabled(true) {}
-	FSimpleGraphFilter(FName InName, bool InEnabled) : Name(InName), bEnabled(InEnabled) {}
-
-	FName Name;
-	uint32 bEnabled : 1;
-};
-
-FORCEINLINE bool operator==(const FSimpleGraphFilter& Left, const FSimpleGraphFilter& Right)
-{
-	return Left.Name == Right.Name;
-}
-
-
 class SFilterWidget;
 class SVisualLoggerFilters : public SCompoundWidget
 {
@@ -31,8 +16,6 @@ public:
 	void AddFilter(const FString& InFilterName);
 	void AddFilter(const FString& GraphName, const FString& DataName);
 	uint32 GetCategoryIndex(const FString& InFilterName) const;
-	bool IsFilterEnabled(const FString& InFilterName, TEnumAsByte<ELogVerbosity::Type> Verbosity = ELogVerbosity::All);
-	bool IsFilterEnabled(const FString& InGraphName, const FString& InDataName, TEnumAsByte<ELogVerbosity::Type> Verbosity = ELogVerbosity::All);
 	void OnFiltersSearchChanged(const FText& Filter);
 	void OnSearchChanged(const FText& Filter);
 	bool GraphSubmenuVisibility(const FName MenuName);
@@ -52,7 +35,7 @@ protected:
 	/** The horizontal box which contains all the filters */
 	TSharedPtr<SWrapBox> FilterBox;
 	TArray<TSharedRef<SFilterWidget> > Filters;
-	TMap<FName, TArray<FSimpleGraphFilter> > GraphFilters;
+	TMap<FName, TArray<FString> > GraphFilters;
 	TSharedPtr<SComboButton> GraphsFilterCombo;
 	FString GraphsFilter;
 };
diff --git a/Engine/Source/Developer/LogVisualizer/Private/SVisualLoggerLogsList.cpp b/Engine/Source/Developer/LogVisualizer/Private/SVisualLoggerLogsList.cpp
--- a/Engine/Source/Developer/LogVisualizer/Private/SVisualLoggerLogsList.cpp
+++ b/Engine/Source/Developer/LogVisualizer/Private/SVisualLoggerLogsList.cpp
@@ -127,7 +127,7 @@ FText SVisualLoggerLogsList::GetFilterText() const
 {
 	static FText NoText;
 	const bool bSearchInsideLogs = ULogVisualizerSettings::StaticClass()->GetDefaultObject<ULogVisualizerSettings>()->bSearchInsideLogs;
-	return bSearchInsideLogs ? FText::FromString(ULogVisualizerSettings::StaticClass()->GetDefaultObject<ULogVisualizerSettings>()->CurrentPresets.DataFilter) : NoText;
+	return bSearchInsideLogs ? FText::FromString(FCategoryFiltersManager::Get().GetSearchString()) : NoText;
 }
 
 void SVisualLoggerLogsList::OnFiltersSearchChanged(const FText& Filter)
@@ -142,19 +142,32 @@ void SVisualLoggerLogsList::OnItemSelectionChanged(const FVisualLogDevice::FVisu
 
 	TArray<FVisualLoggerCategoryVerbosityPair> OutCategories;
 	FVisualLoggerHelpers::GetCategories(LogEntry.Entry, OutCategories);
-	if (FLogVisualizer::Get().GetVisualLoggerInterface()->HasValidCategories(OutCategories) == false)
+	bool bHasValidCategory = false;
+	for (auto& CurrentCategory : OutCategories)
+	{
+		bHasValidCategory |= FCategoryFiltersManager::Get().MatchCategoryFilters(CurrentCategory.CategoryName.ToString(), CurrentCategory.Verbosity);
+	}
+	
+	if (!bHasValidCategory)
 	{
 		return;
 	}
 
 	CurrentLogEntry = LogEntry;
 	const FVisualLogLine* LogLine = LogEntry.Entry.LogLines.GetData();
+	const bool bSearchInsideLogs = ULogVisualizerSettings::StaticClass()->GetDefaultObject<ULogVisualizerSettings>()->bSearchInsideLogs;
 	for (int LineIndex = 0; LineIndex < LogEntry.Entry.LogLines.Num(); ++LineIndex, ++LogLine)
 	{
-		bool bShowLine = true;
-
-		FString CurrentCategory = LogLine->Category.ToString();
-		bShowLine = FLogVisualizer::Get().GetVisualLoggerInterface()->IsValidCategory(CurrentCategory, LogLine->Verbosity) /*&& (bHistogramGraphsFilter || (QuickFilterText.Len() == 0 || CurrentCategory.Find(QuickFilterText) != INDEX_NONE))*/;
+		bool bShowLine = FCategoryFiltersManager::Get().MatchCategoryFilters(LogLine->Category.ToString(), LogLine->Verbosity);
+		if (bSearchInsideLogs)
+		{
+			FString String = FCategoryFiltersManager::Get().GetSearchString();
+			if (String.Len() > 0)
+			{
+				bShowLine &= LogLine->Line.Find(String) != INDEX_NONE || LogLine->Category.ToString().Find(String) != INDEX_NONE;
+			}
+		}
+		
 
 		if (bShowLine)
 		{
@@ -172,10 +185,7 @@ void SVisualLoggerLogsList::OnItemSelectionChanged(const FVisualLogDevice::FVisu
 
 	for (auto& Event : LogEntry.Entry.Events)
 	{
-		bool bShowLine = true;
-
-		FString CurrentCategory = Event.Name;
-		bShowLine = FLogVisualizer::Get().GetVisualLoggerInterface()->IsValidCategory(Event.Name, Event.Verbosity) /*&& (bHistogramGraphsFilter || (QuickFilterText.Len() == 0 || CurrentCategory.Find(QuickFilterText) != INDEX_NONE))*/;
+		bool bShowLine = FCategoryFiltersManager::Get().MatchCategoryFilters(Event.Name, Event.Verbosity);
 
 		if (bShowLine)
 		{
@@ -196,7 +206,6 @@ void SVisualLoggerLogsList::OnItemSelectionChanged(const FVisualLogDevice::FVisu
 
 	}
 
-	//SetCurrentViewedTime(LogEntry.Entry.TimeStamp);
 	LogsLinesWidget->RequestListRefresh();
 }
 #undef LOCTEXT_NAMESPACE
diff --git a/Engine/Source/Developer/LogVisualizer/Private/SVisualLoggerView.cpp b/Engine/Source/Developer/LogVisualizer/Private/SVisualLoggerView.cpp
--- a/Engine/Source/Developer/LogVisualizer/Private/SVisualLoggerView.cpp
+++ b/Engine/Source/Developer/LogVisualizer/Private/SVisualLoggerView.cpp
@@ -99,7 +99,6 @@ void SVisualLoggerView::Construct(const FArguments& InArgs, const TSharedRef<FUI
 
 
 	ULogVisualizerSettings* Settings = ULogVisualizerSettings::StaticClass()->GetDefaultObject<ULogVisualizerSettings>();
-	FFiltersPreset& CurrentPresets = Settings->CurrentPresets;
 
 	ChildSlot
 		[
@@ -129,7 +128,7 @@ void SVisualLoggerView::Construct(const FArguments& InArgs, const TSharedRef<FUI
 							.VAlign(VAlign_Center)
 							[
 								SNew(SImage)
-								.Visibility_Lambda([this]()->EVisibility{ return ULogVisualizerSettings::StaticClass()->GetDefaultObject<ULogVisualizerSettings>()->CurrentPresets.SelectedClasses.Num() > 0 ? EVisibility::Visible : EVisibility::Collapsed; })
+								.Visibility_Lambda([]()->EVisibility{ return FCategoryFiltersManager::Get().GetSelectedObjects().Num() > 0 ? EVisibility::Visible : EVisibility::Collapsed; })
 								.Image(FLogVisualizerStyle::Get().GetBrush("Filters.FilterIcon"))
 							]
 							+ SHorizontalBox::Slot()
@@ -138,7 +137,7 @@ void SVisualLoggerView::Construct(const FArguments& InArgs, const TSharedRef<FUI
 							.AutoWidth()
 							[
 								SAssignNew(ClassesComboButton, SComboButton)
-								.Visibility_Lambda([this]()->EVisibility{ return TimelinesContainer.IsValid() && TimelinesContainer->GetAllNodes().Num() > 1 ? EVisibility::Visible : EVisibility::Collapsed; })
+								.Visibility_Lambda([this]()->EVisibility{ return TimelinesContainer.IsValid() && (TimelinesContainer->GetAllNodes().Num() > 1 || FCategoryFiltersManager::Get().GetSelectedObjects().Num() > 0) ? EVisibility::Visible : EVisibility::Collapsed; })
 								.ComboButtonStyle(FLogVisualizerStyle::Get(), "Filters.Style")
 								.ForegroundColor(FLinearColor::White)
 								.ContentPadding(0)
@@ -220,7 +219,6 @@ void SVisualLoggerView::Construct(const FArguments& InArgs, const TSharedRef<FUI
 						[
 							SNew(SHorizontalBox)
 							+ SHorizontalBox::Slot()
-							//.Visibility(EVisibility::HitTestInvisible)
 							.FillWidth(TAttribute<float>(this, &SVisualLoggerView::GetAnimationOutlinerFillPercentage))
 							[
 								// Take up space but display nothing. This is required so that all areas dependent on time align correctly
@@ -265,8 +263,8 @@ void SVisualLoggerView::Construct(const FArguments& InArgs, const TSharedRef<FUI
 				]
 			]
 		];
-
-		SearchBox->SetText(FText::FromString(CurrentPresets.DataFilter));
+		
+		SearchBox->SetText(FText::FromString(FCategoryFiltersManager::Get().GetSearchString()));
 }
 
 void SVisualLoggerView::SetAnimationOutlinerFillPercentage(float FillPercentage) 
@@ -376,19 +374,21 @@ TSharedRef<SWidget> SVisualLoggerView::MakeClassesFilterMenu()
 				FUIAction(
 				FExecuteAction::CreateLambda([this, OwnerClassName]()
 				{
-					TArray<FString>& SelectedOwnerClasses = ULogVisualizerSettings::StaticClass()->GetDefaultObject<ULogVisualizerSettings>()->CurrentPresets.SelectedClasses;
-					if (SelectedOwnerClasses.Find(OwnerClassName) != INDEX_NONE)
-						SelectedOwnerClasses.Remove(OwnerClassName);
+				if (FCategoryFiltersManager::Get().MatchObjectName(OwnerClassName) && FCategoryFiltersManager::Get().GetSelectedObjects().Num() != 0)
+					{
+						FCategoryFiltersManager::Get().RemoveObjectFromSelection(OwnerClassName);
+					}
 					else
-						SelectedOwnerClasses.AddUnique(OwnerClassName);
+					{
+						FCategoryFiltersManager::Get().SelectObject(OwnerClassName);
+					}
 
 					OnChangedClassesFilter();
 				}),
 				FCanExecuteAction(),
 				FIsActionChecked::CreateLambda([OwnerClassName]()->bool
 				{
-					TArray<FString>& SelectedOwnerClasses = ULogVisualizerSettings::StaticClass()->GetDefaultObject<ULogVisualizerSettings>()->CurrentPresets.SelectedClasses;
-					return SelectedOwnerClasses.Find(OwnerClassName) != INDEX_NONE;
+					return FCategoryFiltersManager::Get().GetSelectedObjects().Find(OwnerClassName) != INDEX_NONE;
 				}),
 				FIsActionButtonVisible()),
 				NAME_None,
@@ -397,7 +397,43 @@ TSharedRef<SWidget> SVisualLoggerView::MakeClassesFilterMenu()
 			UniqueClasses.AddUnique(OwnerClassName);
 		}
 	}
-	MenuBuilder.EndSection(); //ContentBrowserFilterBasicAsset
+	//show any classes from persistent data
+	for (const FString& SelectedObj : FCategoryFiltersManager::Get().GetSelectedObjects())
+	{
+		if (UniqueClasses.Find(SelectedObj) == INDEX_NONE)
+		{
+			FText LabelText = FText::FromString(SelectedObj);
+			MenuBuilder.AddMenuEntry(
+			LabelText,
+			FText::Format(LOCTEXT("FilterByClassPrefix", "Toggle {0} class"), LabelText),
+			FSlateIcon(),
+			FUIAction(
+			FExecuteAction::CreateLambda([this, SelectedObj]()
+			{
+				if (FCategoryFiltersManager::Get().MatchObjectName(SelectedObj) && FCategoryFiltersManager::Get().GetSelectedObjects().Num() != 0)
+				{
+					FCategoryFiltersManager::Get().RemoveObjectFromSelection(SelectedObj);
+				}
+				else
+				{
+					FCategoryFiltersManager::Get().SelectObject(SelectedObj);
+				}
+
+				OnChangedClassesFilter();
+			}),
+			FCanExecuteAction(),
+			FIsActionChecked::CreateLambda([SelectedObj]()->bool
+			{
+				return FCategoryFiltersManager::Get().GetSelectedObjects().Find(SelectedObj) != INDEX_NONE;
+			}),
+			FIsActionButtonVisible()),
+			NAME_None,
+			EUserInterfaceActionType::ToggleButton
+			);
+			UniqueClasses.AddUnique(SelectedObj);
+		}
+	}
+	MenuBuilder.EndSection(); 
 
 
 	FDisplayMetrics DisplayMetrics;
diff --git a/Engine/Source/Developer/LogVisualizer/Private/VisualLoggerCanvasRenderer.cpp b/Engine/Source/Developer/LogVisualizer/Private/VisualLoggerCanvasRenderer.cpp
--- a/Engine/Source/Developer/LogVisualizer/Private/VisualLoggerCanvasRenderer.cpp
+++ b/Engine/Source/Developer/LogVisualizer/Private/VisualLoggerCanvasRenderer.cpp
@@ -84,7 +84,7 @@ void FVisualLoggerCanvasRenderer::DrawOnCanvas(class UCanvas* Canvas, class APla
 	for (const auto CurrentData : SelectedEntry.DataBlocks)
 	{
 		const FName TagName = CurrentData.TagName;
-		const bool bIsValidByFilter = FLogVisualizer::Get().GetVisualLoggerInterface()->IsValidCategory(CurrentData.Category.ToString(), ELogVerbosity::All) && FLogVisualizer::Get().GetVisualLoggerInterface()->IsValidCategory(CurrentData.TagName.ToString(), ELogVerbosity::All);
+		const bool bIsValidByFilter = FCategoryFiltersManager::Get().MatchCategoryFilters(CurrentData.Category.ToString(), ELogVerbosity::All) && FCategoryFiltersManager::Get().MatchCategoryFilters(CurrentData.TagName.ToString(), ELogVerbosity::All);
 		FVisualLogExtensionInterface* Extension = FVisualLogger::Get().GetExtensionForTag(TagName);
 		if (!Extension)
 		{
@@ -169,7 +169,9 @@ void FVisualLoggerCanvasRenderer::DrawHistogramGraphs(class UCanvas* Canvas, cla
 			const FName CurrentGraphName = CurrentSample.GraphName;
 			const FName CurrentDataName = CurrentSample.DataName;
 
-			const bool bIsValidByFilter = FLogVisualizer::Get().GetVisualLoggerInterface()->IsValidCategory(CurrentSample.GraphName.ToString(), CurrentSample.DataName.ToString(), ELogVerbosity::All);
+			FString GraphFilterName = CurrentSample.GraphName.ToString() +TEXT("$") + CurrentSample.DataName.ToString();
+			const bool bIsValidByFilter = FCategoryFiltersManager::Get().MatchCategoryFilters(GraphFilterName, ELogVerbosity::All);
+
 			if (bIsValidByFilter)
 			{
 				FGraphData &GraphData = CollectedGraphs.FindOrAdd(CurrentSample.GraphName);
@@ -199,7 +201,8 @@ void FVisualLoggerCanvasRenderer::DrawHistogramGraphs(class UCanvas* Canvas, cla
 			const FName CurrentGraphName = CurrentSample.GraphName;
 			const FName CurrentDataName = CurrentSample.DataName;
 
-			const bool bIsValidByFilter = FLogVisualizer::Get().GetVisualLoggerInterface()->IsValidCategory(CurrentSample.GraphName.ToString(), CurrentSample.DataName.ToString(), ELogVerbosity::All);
+			FString GraphFilterName = CurrentSample.GraphName.ToString() + TEXT("_") + CurrentSample.DataName.ToString();
+			const bool bIsValidByFilter = FCategoryFiltersManager::Get().MatchCategoryFilters(GraphFilterName, ELogVerbosity::All);
 			if (bIsValidByFilter)
 			{
 				FGraphData &GraphData = CollectedGraphs.FindOrAdd(CurrentSample.GraphName);
diff --git a/Engine/Source/Developer/LogVisualizer/Private/VisualLoggerRenderingActor.cpp b/Engine/Source/Developer/LogVisualizer/Private/VisualLoggerRenderingActor.cpp
--- a/Engine/Source/Developer/LogVisualizer/Private/VisualLoggerRenderingActor.cpp
+++ b/Engine/Source/Developer/LogVisualizer/Private/VisualLoggerRenderingActor.cpp
@@ -229,7 +229,7 @@ void AVisualLoggerRenderingActor::OnItemSelectionChanged(const FVisualLogDevice:
 	for (const auto CurrentData : Entry->DataBlocks)
 	{
 		const FName TagName = CurrentData.TagName;
-		const bool bIsValidByFilter = FLogVisualizer::Get().GetVisualLoggerInterface()->IsValidCategory(CurrentData.Category.ToString(), ELogVerbosity::All) && FLogVisualizer::Get().GetVisualLoggerInterface()->IsValidCategory(CurrentData.TagName.ToString(), ELogVerbosity::All);
+		const bool bIsValidByFilter = FCategoryFiltersManager::Get().MatchCategoryFilters(CurrentData.Category.ToString(), ELogVerbosity::All) && FCategoryFiltersManager::Get().MatchCategoryFilters(CurrentData.TagName.ToString(), ELogVerbosity::All);
 		FVisualLogExtensionInterface* Extension = FVisualLogger::Get().GetExtensionForTag(TagName);
 		if (!Extension)
 		{
@@ -248,7 +248,7 @@ void AVisualLoggerRenderingActor::OnItemSelectionChanged(const FVisualLogDevice:
 
 	for (int32 ElementIndex = 0; ElementIndex < ElementsCount; ++ElementIndex, ++ElementToDraw)
 	{
-		if (!FLogVisualizer::Get().GetVisualLoggerInterface()->IsValidCategory(ElementToDraw->Category.ToString(), ElementToDraw->Verbosity))
+		if (!FCategoryFiltersManager::Get().MatchCategoryFilters(ElementToDraw->Category.ToString(), ElementToDraw->Verbosity))
 		{
 			continue;
 		}
diff --git a/Engine/Source/Developer/LogVisualizer/Public/LogVisualizerSettings.h b/Engine/Source/Developer/LogVisualizer/Public/LogVisualizerSettings.h
--- a/Engine/Source/Developer/LogVisualizer/Public/LogVisualizerSettings.h
+++ b/Engine/Source/Developer/LogVisualizer/Public/LogVisualizerSettings.h
@@ -9,30 +9,46 @@ GameplayDebuggerSettings.h: Declares the UGameplayDebuggerSettings class.
 #include "LogVisualizerSettings.generated.h"
 
 USTRUCT()
-struct FFiltersPreset
+struct FCategoryFilter
 {
 	GENERATED_USTRUCT_BODY()
 
 	UPROPERTY(config)
-	FString FilterName;
+	FString CategoryName;
 
 	UPROPERTY(config)
-	FString ObjectNameFilter;
+	int32 LogVerbosity;
 
 	UPROPERTY(config)
-	FString DataFilter;
+	bool Enabled;
+};
 
-	UPROPERTY(config)
-	TArray<FString> CategoryFilters;
+USTRUCT()
+struct FVisualLoggerFilters
+{
+	GENERATED_USTRUCT_BODY()
 
+	UPROPERTY(config)
+	FString SearchBoxFilter;
+	
+	UPROPERTY(config)
+	FString ObjectNameFilter;
+	
+	UPROPERTY(config)
+	TArray<FCategoryFilter> Categories;
+	
 	UPROPERTY(config)
 	TArray<FString> SelectedClasses;
 };
 
+struct FCategoryFiltersManager;
+
 UCLASS(config = EditorUserSettings)
 class LOGVISUALIZER_API ULogVisualizerSettings : public UObject
 {
 	GENERATED_UCLASS_BODY()
+	friend struct FCategoryFiltersManager;
+
 public:
 	DECLARE_EVENT_OneParam(ULogVisualizerSettings, FSettingChangedEvent, FName /*PropertyName*/);
 	FSettingChangedEvent& OnSettingChanged() { return SettingChangedEvent; }
@@ -69,9 +85,6 @@ public:
 	UPROPERTY(EditAnywhere, config, Category = "VisualLogger")
 	FColor GraphsBackgroundColor;
 
-	UPROPERTY(config)
-	TArray<FFiltersPreset> FilterPresets;
-
 	/**Whether to store all filter settings on exit*/
 	UPROPERTY(EditAnywhere, config, Category = "VisualLogger")
 	bool bPresistentFilters;
@@ -84,19 +97,51 @@ public:
 	UPROPERTY(EditAnywhere, config, Category = "VisualLogger")
 	bool bUsePlayersOnlyForPause;
 
-	UPROPERTY(config)
-	FFiltersPreset PresistentPresets;
-
-	FFiltersPreset CurrentPresets;
-
 	// UObject overrides
 #if WITH_EDITOR
 	virtual void PostEditChangeProperty(struct FPropertyChangedEvent& PropertyChangedEvent) override;
 #endif
 
+protected:
+	FVisualLoggerFilters CurrentFilters;
+	
+	UPROPERTY(config)
+	FVisualLoggerFilters PresistentFilters;
+
 private:
 
 	// Holds an event delegate that is executed when a setting has changed.
 	FSettingChangedEvent SettingChangedEvent;
 
 };
+
+struct FCategoryFiltersManager
+{
+	static FCategoryFiltersManager& Get() { return StaticManager; }
+
+	bool MatchCategoryFilters(FString String, ELogVerbosity::Type Verbosity = ELogVerbosity::All);
+	bool MatchObjectName(FString String);
+	bool MatchSearchString(FString String);
+
+	void SetSearchString(FString InString);
+	FString GetSearchString();
+
+	void SetObjectFilterString(FString InFilterString);
+	FString GetObjectFilterString();
+
+	void AddCategory(FString InName, ELogVerbosity::Type InVerbosity);
+	void RemoveCategory(FString InName);
+	bool IsValidCategory(FString InName);
+	FCategoryFilter& GetCategory(FString InName);
+
+	void SelectObject(FString ObjectName);
+	void RemoveObjectFromSelection(FString ObjectName);
+	const TArray<FString>& GetSelectedObjects();
+
+	void SavePresistentData();
+	void ClearPresistentData();
+	void LoadPresistentData();
+
+protected:
+	static FCategoryFiltersManager StaticManager;
+};
