Fixed some bad font references

UE-8718 - Warning about loading font darta from Roboto-Regular.ttf & Roboto-Bold.ttf

diff --git a/Engine/Source/Developer/LogVisualizer/Private/LogVisualizerStyle.cpp b/Engine/Source/Developer/LogVisualizer/Private/LogVisualizerStyle.cpp
--- a/Engine/Source/Developer/LogVisualizer/Private/LogVisualizerStyle.cpp
+++ b/Engine/Source/Developer/LogVisualizer/Private/LogVisualizerStyle.cpp
@@ -47,7 +47,7 @@ TSharedRef< FSlateStyleSet > FLogVisualizerStyle::Create()
 	FSlateStyleSet& Style = StyleRef.Get();
 	// Generic styles
 	const FTextBlockStyle NormalText = FTextBlockStyle()
-		.SetFont(TTF_FONT("Fonts/Roboto-Regular", 9))
+		.SetFont(TTF_CORE_FONT("Fonts/Roboto-Regular", 9))
 		.SetColorAndOpacity(FSlateColor::UseForeground())
 		.SetShadowOffset(FVector2D::ZeroVector)
 		.SetShadowColorAndOpacity(FLinearColor::Black)
@@ -128,7 +128,7 @@ TSharedRef< FSlateStyleSet > FLogVisualizerStyle::Create()
 		Style.Set("Sequencer.ItemTitle.Hover", new BOX_BRUSH("Common/Button/simple_round_hovered", FMargin(4 / 16.0f), FLinearColor(1, 1, 1, 1)));
 		Style.Set("Sequencer.SectionArea.Background", new FSlateColorBrush(FColor::White));
 		const FTextBlockStyle CLassNameTextStyle = FTextBlockStyle(NormalText)
-			.SetFont(TTF_FONT("Fonts/Roboto-Regular", 7));
+			.SetFont(TTF_CORE_FONT("Fonts/Roboto-Regular", 7));
 		Style.Set("Sequencer.ClassNAme", CLassNameTextStyle);
 	}
 
@@ -139,7 +139,7 @@ TSharedRef< FSlateStyleSet > FLogVisualizerStyle::Create()
 		Style.Set("RichText.RoundedBackground", new BOX_BRUSH("Common/RoundedSelection_16x", 4.0f / 16.0f, FLinearColor(FColor(0xffeff3f3))));
 
 		const FTextBlockStyle NormalRichTextStyle = FTextBlockStyle(NormalText)
-			.SetFont(TTF_FONT("Fonts/Roboto-Regular", 10))
+			.SetFont(TTF_CORE_FONT("Fonts/Roboto-Regular", 10))
 			.SetColorAndOpacity(FLinearColor(FColor(0xff2c3e50)));
 
 		Style.Set("RichText.Text", NormalRichTextStyle);
@@ -187,22 +187,22 @@ TSharedRef< FSlateStyleSet > FLogVisualizerStyle::Create()
 			Style.Set("RichText.Tagline.Background", new BOX_BRUSH("Common/FlatColorSquare", FVector2D(1.0f, 1.0f), FMargin(1), FLinearColor(FColor(0xffdbe4e4))));
 			Style.Set("RichText.Tagline.DarkBackground", new BOX_BRUSH("Common/FlatColorSquare", FVector2D(1.0f, 1.0f), FMargin(1), FLinearColor(0.55423, 0.60548, 0.60548)));
 			Style.Set("RichText.Tagline.Text", FTextBlockStyle(NormalText)
-				.SetFont(TTF_FONT("Fonts/Roboto-Bold", 24))
+				.SetFont(TTF_CORE_FONT("Fonts/Roboto-Bold", 24))
 				.SetColorAndOpacity(FLinearColor(FColor(0xff2c3e50)))
 				);
 
 			Style.Set("RichText.Tagline.TextHighlight", FTextBlockStyle(NormalText)
-				.SetFont(TTF_FONT("Fonts/Roboto-Bold", 24))
+				.SetFont(TTF_CORE_FONT("Fonts/Roboto-Bold", 24))
 				.SetColorAndOpacity(FLinearColor(FColor(0xff19bc9c)))
 				);
 
 			Style.Set("RichText.Tagline.SubtleText", FTextBlockStyle(NormalText)
-				.SetFont(TTF_FONT("Fonts/Roboto-Regular", 16))
+				.SetFont(TTF_CORE_FONT("Fonts/Roboto-Regular", 16))
 				.SetColorAndOpacity(FLinearColor(FColor(0xff2c3e50)))
 				);
 
 			Style.Set("RichText.Tagline.SubtleTextHighlight", FTextBlockStyle(NormalText)
-				.SetFont(TTF_FONT("Fonts/Roboto-Regular", 16))
+				.SetFont(TTF_CORE_FONT("Fonts/Roboto-Regular", 16))
 				.SetColorAndOpacity(FLinearColor(FColor(0xff19bc9c)))
 				);
 		}
@@ -217,6 +217,7 @@ TSharedRef< FSlateStyleSet > FLogVisualizerStyle::Create()
 #undef BORDER_BRUSH
 #undef TTF_FONT
 #undef OTF_FONT
+#undef TTF_CORE_FONT
 
 const ISlateStyle& FLogVisualizerStyle::Get()
 {
diff --git a/Engine/Source/Editor/AddContentDialog/Private/AddContentDialogStyle.cpp b/Engine/Source/Editor/AddContentDialog/Private/AddContentDialogStyle.cpp
--- a/Engine/Source/Editor/AddContentDialog/Private/AddContentDialogStyle.cpp
+++ b/Engine/Source/Editor/AddContentDialog/Private/AddContentDialogStyle.cpp
@@ -33,6 +33,7 @@ FName FAddContentDialogStyle::GetStyleSetName()
 #define BORDER_BRUSH( RelativePath, ... ) FSlateBorderBrush( Style->RootToContentDir( RelativePath, TEXT(".png") ), __VA_ARGS__ )
 #define TTF_FONT( RelativePath, ... ) FSlateFontInfo( Style->RootToContentDir( RelativePath, TEXT(".ttf") ), __VA_ARGS__ )
 #define OTF_FONT( RelativePath, ... ) FSlateFontInfo( Style->RootToContentDir( RelativePath, TEXT(".otf") ), __VA_ARGS__ )
+#define TTF_CORE_FONT( RelativePath, ... ) FSlateFontInfo( FPaths::EngineContentDir()  / "Slate"/ RelativePath + TEXT(".ttf"), __VA_ARGS__ )
 
 const FVector2D Icon16x16(16.0f, 16.0f);
 const FVector2D Icon20x20(20.0f, 20.0f);
@@ -66,12 +67,12 @@ TSharedRef< FSlateStyleSet > FAddContentDialogStyle::Create()
 	Style->Set("AddContentDialog.RightArrow", new IMAGE_BRUSH("ArrowRight", FVector2D(25.0f, 42.0f)));
 
 	Style->Set("AddContentDialog.HeadingText", FTextBlockStyle(NormalText)
-		.SetFont(TTF_FONT("Fonts/Roboto-Regular", 14))
+		.SetFont(TTF_CORE_FONT("Fonts/Roboto-Regular", 14))
 		.SetShadowOffset(FVector2D(0, 1))
 		.SetShadowColorAndOpacity(FLinearColor(0, 0, 0, 0.9f)));
 
 	Style->Set("AddContentDialog.HeadingTextSmall", FTextBlockStyle(NormalText)
-		.SetFont(TTF_FONT("Fonts/Roboto-Regular", 12))
+		.SetFont(TTF_CORE_FONT("Fonts/Roboto-Regular", 12))
 		.SetShadowOffset(FVector2D(0, 1))
 		.SetShadowColorAndOpacity(FLinearColor(0, 0, 0, 0.9f)));
 
@@ -93,6 +94,7 @@ TSharedRef< FSlateStyleSet > FAddContentDialogStyle::Create()
 #undef BORDER_BRUSH
 #undef TTF_FONT
 #undef OTF_FONT
+#undef TTF_CORE_FONT
 
 void FAddContentDialogStyle::ReloadTextures()
 {
diff --git a/Engine/Source/Editor/EditorStyle/Private/SlateEditorStyle.cpp b/Engine/Source/Editor/EditorStyle/Private/SlateEditorStyle.cpp
--- a/Engine/Source/Editor/EditorStyle/Private/SlateEditorStyle.cpp
+++ b/Engine/Source/Editor/EditorStyle/Private/SlateEditorStyle.cpp
@@ -4029,7 +4029,7 @@ void FSlateEditorStyle::FStyle::SetupGraphEditorStyles()
 
 		Set("SCSEditor.ComponentTooltip.ClassDescription",
 			FTextBlockStyle(NormalText)
-			.SetFont(TTF_CORE_FONT("Fonts/Roboto-Italic", 10))
+			.SetFont(TTF_FONT("Fonts/Roboto-Italic", 10))
 			.SetColorAndOpacity(FLinearColor(0.1f, 0.1f, 0.1f))
 			.SetShadowOffset(FVector2D(1.0f, 1.0f))
 			.SetShadowColorAndOpacity(FLinearColor(0.8f, 0.8f, 0.8f))
