Fixed EQS rendering in Visual Logger tool, in editor.

diff --git a/Engine/Source/Developer/LogVisualizer/Private/LogVisualizer.h b/Engine/Source/Developer/LogVisualizer/Private/LogVisualizer.h
--- a/Engine/Source/Developer/LogVisualizer/Private/LogVisualizer.h
+++ b/Engine/Source/Developer/LogVisualizer/Private/LogVisualizer.h
@@ -21,6 +21,7 @@ DECLARE_DELEGATE_OneParam(FOnItemSelectionChanged, const FVisualLogDevice::FVisu
 DECLARE_DELEGATE_OneParam(FOnObjectSelectionChanged, TSharedPtr<class STimeline>);
 DECLARE_DELEGATE_OneParam(FOnFiltersSearchChanged, const FText&);
 DECLARE_DELEGATE(FOnFiltersChanged);
+DECLARE_DELEGATE_ThreeParams(FOnLogLineSelectionChanged, TSharedPtr<FLogEntryItem> /*SelectedItem*/, int64 /*UserData*/, FName /*TagName*/);
 
 struct FVisualLoggerEvents
 {
@@ -32,6 +33,7 @@ struct FVisualLoggerEvents
 	FOnItemSelectionChanged OnItemSelectionChanged;
 	FOnFiltersChanged OnFiltersChanged;
 	FOnObjectSelectionChanged OnObjectSelectionChanged;
+	FOnLogLineSelectionChanged OnLogLineSelectionChanged;
 };
 
 class FVisualLoggerTimeSliderController;
diff --git a/Engine/Source/Developer/LogVisualizer/Private/STimelineBar.cpp b/Engine/Source/Developer/LogVisualizer/Private/STimelineBar.cpp
--- a/Engine/Source/Developer/LogVisualizer/Private/STimelineBar.cpp
+++ b/Engine/Source/Developer/LogVisualizer/Private/STimelineBar.cpp
@@ -485,6 +485,7 @@ int32 STimelineBar::OnPaint(const FPaintArgs& Args, const FGeometry& AllottedGeo
 		if (BestItemIndex != CurrentItemIndex)
 		{
 			FLogVisualizer::Get().GetVisualLoggerEvents().OnItemSelectionChanged.ExecuteIfBound(BestItemEntry);
+			CurrentItemIndex = BestItemIndex;
 		}
 
 		float CurrentTime = BestItemEntry.Entry.TimeStamp;
diff --git a/Engine/Source/Developer/LogVisualizer/Private/SVisualLogger.cpp b/Engine/Source/Developer/LogVisualizer/Private/SVisualLogger.cpp
--- a/Engine/Source/Developer/LogVisualizer/Private/SVisualLogger.cpp
+++ b/Engine/Source/Developer/LogVisualizer/Private/SVisualLogger.cpp
@@ -129,6 +129,7 @@ void SVisualLogger::Construct(const FArguments& InArgs, const TSharedRef<SDockTa
 	FLogVisualizer::Get().GetVisualLoggerEvents().OnItemSelectionChanged = FOnItemSelectionChanged::CreateRaw(this, &SVisualLogger::OnItemSelectionChanged);
 	FLogVisualizer::Get().GetVisualLoggerEvents().OnFiltersChanged = FOnFiltersChanged::CreateRaw(this, &SVisualLogger::OnFiltersChanged);
 	FLogVisualizer::Get().GetVisualLoggerEvents().OnObjectSelectionChanged = FOnObjectSelectionChanged::CreateRaw(this, &SVisualLogger::OnObjectSelectionChanged);
+	FLogVisualizer::Get().GetVisualLoggerEvents().OnLogLineSelectionChanged = FOnLogLineSelectionChanged::CreateRaw(this, &SVisualLogger::OnLogLineSelectionChanged);
 
 	//////////////////////////////////////////////////////////////////////////
 	// Command Action Lists
@@ -811,6 +812,26 @@ void SVisualLogger::OnFiltersSearchChanged(const FText& Filter)
 	}
 }
 
+void SVisualLogger::OnLogLineSelectionChanged(TSharedPtr<struct FLogEntryItem> SelectedItem, int64 UserData, FName TagName)
+{
+	TMap<FName, FVisualLogExtensionInterface*>& AllExtensions = FVisualLogger::Get().GetAllExtensions();
+	for (auto Iterator = AllExtensions.CreateIterator(); Iterator; ++Iterator)
+	{
+		FVisualLogExtensionInterface* Extension = (*Iterator).Value;
+		if (Extension != NULL)
+		{
+			Extension->LogEntryLineSelectionChanged(SelectedItem, UserData, TagName);
+		}
+	}
+
+	AVisualLoggerRenderingActor* HelperActor = Cast<AVisualLoggerRenderingActor>(FLogVisualizer::Get().GetVisualLoggerHelperActor());
+	if (HelperActor)
+	{
+		HelperActor->OnItemSelectionChanged(LogsList->GetCurrentLogEntry());
+		HelperActor->MarkComponentsRenderStateDirty();
+	}
+}
+
 void SVisualLogger::OnMoveCursorLeftCommand()
 {
 	FLogVisualizer::Get().GotoNextItem();
diff --git a/Engine/Source/Developer/LogVisualizer/Private/SVisualLogger.h b/Engine/Source/Developer/LogVisualizer/Private/SVisualLogger.h
--- a/Engine/Source/Developer/LogVisualizer/Private/SVisualLogger.h
+++ b/Engine/Source/Developer/LogVisualizer/Private/SVisualLogger.h
@@ -68,6 +68,7 @@ public:
 	void OnFiltersChanged();
 	void OnObjectSelectionChanged(TSharedPtr<class STimeline> TimeLine);
 	void OnFiltersSearchChanged(const FText& Filter);
+	void OnLogLineSelectionChanged(TSharedPtr<struct FLogEntryItem> SelectedItem, int64 UserData, FName TagName);
 
 	bool HandleStartRecordingCommandCanExecute() const;
 	void HandleStartRecordingCommandExecute();
diff --git a/Engine/Source/Developer/LogVisualizer/Private/SVisualLoggerLogsList.cpp b/Engine/Source/Developer/LogVisualizer/Private/SVisualLoggerLogsList.cpp
--- a/Engine/Source/Developer/LogVisualizer/Private/SVisualLoggerLogsList.cpp
+++ b/Engine/Source/Developer/LogVisualizer/Private/SVisualLoggerLogsList.cpp
@@ -71,21 +71,13 @@ void SVisualLoggerLogsList::OnFiltersChanged()
 
 void SVisualLoggerLogsList::LogEntryLineSelectionChanged(TSharedPtr<FLogEntryItem> SelectedItem, ESelectInfo::Type SelectInfo)
 {
-	TMap<FName, FVisualLogExtensionInterface*>& AllExtensions = FVisualLogger::Get().GetAllExtensions();
-	for (auto Iterator = AllExtensions.CreateIterator(); Iterator; ++Iterator)
+	if (SelectedItem.IsValid() == true)
 	{
-		FVisualLogExtensionInterface* Extension = (*Iterator).Value;
-		if (Extension != NULL)
-		{
-			if (SelectedItem.IsValid() == true)
-			{
-				Extension->LogEntryLineSelectionChanged(SelectedItem, SelectedItem->UserData, SelectedItem->TagName);
-			}
-			else
-			{
-				Extension->LogEntryLineSelectionChanged(SelectedItem, 0, NAME_None);
-			}
-		}
+		FLogVisualizer::Get().GetVisualLoggerEvents().OnLogLineSelectionChanged.ExecuteIfBound(SelectedItem, SelectedItem->UserData, SelectedItem->TagName);
+	}
+	else
+	{
+		FLogVisualizer::Get().GetVisualLoggerEvents().OnLogLineSelectionChanged.ExecuteIfBound(SelectedItem, 0, NAME_None);
 	}
 }
 
diff --git a/Engine/Source/Developer/LogVisualizer/Private/SVisualLoggerLogsList.h b/Engine/Source/Developer/LogVisualizer/Private/SVisualLoggerLogsList.h
--- a/Engine/Source/Developer/LogVisualizer/Private/SVisualLoggerLogsList.h
+++ b/Engine/Source/Developer/LogVisualizer/Private/SVisualLoggerLogsList.h
@@ -19,6 +19,7 @@ public:
 	void OnFiltersChanged();
 	void OnFiltersSearchChanged(const FText& Filter);
 	FText GetFilterText() const;
+	const FVisualLogDevice::FVisualLogEntryItem& GetCurrentLogEntry() { return CurrentLogEntry; }
 
 protected:
 	TSharedPtr<SListView<TSharedPtr<struct FLogEntryItem> > > LogsLinesWidget;
diff --git a/Engine/Source/Runtime/AIModule/Classes/VisualLoggerExtension.h b/Engine/Source/Runtime/AIModule/Classes/VisualLoggerExtension.h
--- a/Engine/Source/Runtime/AIModule/Classes/VisualLoggerExtension.h
+++ b/Engine/Source/Runtime/AIModule/Classes/VisualLoggerExtension.h
@@ -10,7 +10,7 @@ class AActor;
 
 namespace EVisLogTags
 {
-	const FString TAG_EQS = TEXT("EQS");
+	const FString TAG_EQS = TEXT("LogEQS");
 }
 
 #if ENABLE_VISUAL_LOG
diff --git a/Engine/Source/Runtime/AIModule/Private/VisualLoggerExtension.cpp b/Engine/Source/Runtime/AIModule/Private/VisualLoggerExtension.cpp
--- a/Engine/Source/Runtime/AIModule/Private/VisualLoggerExtension.cpp
+++ b/Engine/Source/Runtime/AIModule/Private/VisualLoggerExtension.cpp
@@ -85,20 +85,17 @@ void FVisualLoggerExtension::DrawData(UWorld* InWorld, UCanvas* Canvas, AActor*
 				EQSRenderComp->SetHiddenInGame(true);
 			}
 
-			if (DebugData.Id != CachedEQSId || (EQSRenderComp && EQSRenderComp->bHiddenInGame))
+			if (SelectedEQSId == INDEX_NONE || SelectedEQSId != CachedEQSId /*|| (EQSRenderComp && EQSRenderComp->bHiddenInGame)*/)
 			{
 				if (SelectedEQSId == INDEX_NONE)
 				{
 					SelectedEQSId = DebugData.Id;
 				}
-				if (DebugData.Id == SelectedEQSId)
-				{
-					CachedEQSId = DebugData.Id;
-					EQSRenderComp->DebugData = DebugData;
-					EQSRenderComp->Activate();
-					EQSRenderComp->SetHiddenInGame(false);
-					EQSRenderComp->MarkRenderStateDirty();
-				}
+				CachedEQSId = DebugData.Id;
+				EQSRenderComp->DebugData = DebugData;
+				EQSRenderComp->Activate();
+				EQSRenderComp->SetHiddenInGame(false);
+				EQSRenderComp->MarkRenderStateDirty();
 			}
 		}
 
