Different optimizations to drawing histogram data from Visual Logger.

diff --git a/Engine/Source/Developer/LogVisualizer/Private/STimelinesContainer.cpp b/Engine/Source/Developer/LogVisualizer/Private/STimelinesContainer.cpp
--- a/Engine/Source/Developer/LogVisualizer/Private/STimelinesContainer.cpp
+++ b/Engine/Source/Developer/LogVisualizer/Private/STimelinesContainer.cpp
@@ -165,7 +165,9 @@ FReply STimelinesContainer::OnMouseWheel(const FGeometry& MyGeometry, const FPoi
 {
 	if (MouseEvent.IsLeftControlDown() || MouseEvent.IsLeftShiftDown())
 	{
-		return TimeSliderController->OnMouseWheel(SharedThis(this), MyGeometry, MouseEvent);
+		FReply RetValue = TimeSliderController->OnMouseWheel(SharedThis(this), MyGeometry, MouseEvent);
+		FLogVisualizer::Get().GetVisualLoggerEvents().OnFiltersChanged.ExecuteIfBound();
+		return RetValue;
 	}
 	return FReply::Unhandled();
 }
diff --git a/Engine/Source/Developer/LogVisualizer/Private/SVisualLogger.cpp b/Engine/Source/Developer/LogVisualizer/Private/SVisualLogger.cpp
--- a/Engine/Source/Developer/LogVisualizer/Private/SVisualLogger.cpp
+++ b/Engine/Source/Developer/LogVisualizer/Private/SVisualLogger.cpp
@@ -785,6 +785,7 @@ void SVisualLogger::OnFiltersChanged()
 		HelperActor->OnItemSelectionChanged(LogsList->GetCurrentLogEntry());
 		HelperActor->MarkComponentsRenderStateDirty();
 	}
+	VisualLoggerCanvasRenderer->OnItemSelectionChanged(LogsList->GetCurrentLogEntry().Entry);
 }
 
 void SVisualLogger::OnObjectSelectionChanged(TSharedPtr<class STimeline> TimeLine)
diff --git a/Engine/Source/Developer/LogVisualizer/Private/SVisualLoggerFilters.cpp b/Engine/Source/Developer/LogVisualizer/Private/SVisualLoggerFilters.cpp
--- a/Engine/Source/Developer/LogVisualizer/Private/SVisualLoggerFilters.cpp
+++ b/Engine/Source/Developer/LogVisualizer/Private/SVisualLoggerFilters.cpp
@@ -150,6 +150,7 @@ void SVisualLoggerFilters::GraphFilterCategoryClicked(FName MenuCategory)
 
 		if (bChanged)
 		{
+			FLogVisualizer::Get().GetVisualLoggerEvents().OnFiltersChanged.ExecuteIfBound();
 			InvalidateCanvas();
 		}
 	}
@@ -211,6 +212,7 @@ void SVisualLoggerFilters::FilterByTypeClicked(FName InGraphName, FName InDataNa
 
 		if (bChanged)
 		{
+			FLogVisualizer::Get().GetVisualLoggerEvents().OnFiltersChanged.ExecuteIfBound();
 			InvalidateCanvas();
 		}
 	}
diff --git a/Engine/Source/Developer/LogVisualizer/Private/TimeSliderController.cpp b/Engine/Source/Developer/LogVisualizer/Private/TimeSliderController.cpp
--- a/Engine/Source/Developer/LogVisualizer/Private/TimeSliderController.cpp
+++ b/Engine/Source/Developer/LogVisualizer/Private/TimeSliderController.cpp
@@ -562,7 +562,7 @@ FReply FVisualLoggerTimeSliderController::OnMouseWheel( TSharedRef<SWidget> Widg
 
 	if (MouseEvent.IsLeftShiftDown())
 	{
-		const float ZoomDelta = 0.1f * MouseEvent.GetWheelDelta();
+		const float ZoomDelta = 0.025f * MouseEvent.GetWheelDelta();
 		TimeSliderArgs.CursorSize.Set(FMath::Clamp(TimeSliderArgs.CursorSize.Get() + ZoomDelta, 0.0f, 1.0f));
 
 		ReturnValue = FReply::Handled();
diff --git a/Engine/Source/Developer/LogVisualizer/Private/TimeSliderController.h b/Engine/Source/Developer/LogVisualizer/Private/TimeSliderController.h
--- a/Engine/Source/Developer/LogVisualizer/Private/TimeSliderController.h
+++ b/Engine/Source/Developer/LogVisualizer/Private/TimeSliderController.h
@@ -7,7 +7,7 @@
 struct FVisualLoggerTimeSliderArgs : FTimeSliderArgs
 {
 	FVisualLoggerTimeSliderArgs()
-		: CursorSize(0.1f)
+		: CursorSize(0.05f)
 	{
 		ViewRange = TRange<float>(0.0f, 5.0f);
 	}
diff --git a/Engine/Source/Developer/LogVisualizer/Private/VisualLoggerCanvasRenderer.cpp b/Engine/Source/Developer/LogVisualizer/Private/VisualLoggerCanvasRenderer.cpp
--- a/Engine/Source/Developer/LogVisualizer/Private/VisualLoggerCanvasRenderer.cpp
+++ b/Engine/Source/Developer/LogVisualizer/Private/VisualLoggerCanvasRenderer.cpp
@@ -55,11 +55,13 @@ namespace LogVisualizer
 void FVisualLoggerCanvasRenderer::OnItemSelectionChanged(const struct FVisualLogEntry& EntryItem)
 {
 	SelectedEntry = EntryItem;
+	DirtyCachedData();
 }
 
 void FVisualLoggerCanvasRenderer::ObjectSelectionChanged(TSharedPtr<class STimeline> TimeLine)
 {
 	CurrentTimeLine = TimeLine;
+	DirtyCachedData();
 }
 
 void FVisualLoggerCanvasRenderer::DrawOnCanvas(class UCanvas* Canvas, class APlayerController*)
@@ -81,30 +83,44 @@ void FVisualLoggerCanvasRenderer::DrawOnCanvas(class UCanvas* Canvas, class APla
 	const FString TimeStampString = FString::Printf(TEXT("%.2f"), SelectedEntry.TimeStamp);
 	LogVisualizer::DrawTextShadowed(Canvas, Font, TimeStampString, SelectedEntry.Location);
 
-	for (const auto CurrentData : SelectedEntry.DataBlocks)
+	if (bDirtyData)
 	{
-		const FName TagName = CurrentData.TagName;
-		const bool bIsValidByFilter = FCategoryFiltersManager::Get().MatchCategoryFilters(CurrentData.Category.ToString(), ELogVerbosity::All) && FCategoryFiltersManager::Get().MatchCategoryFilters(CurrentData.TagName.ToString(), ELogVerbosity::All);
-		FVisualLogExtensionInterface* Extension = FVisualLogger::Get().GetExtensionForTag(TagName);
-		if (!Extension)
+		for (const auto& CurrentData : CachedDataBlocks)
 		{
-			continue;
+			if (FVisualLogExtensionInterface* Extension = FVisualLogger::Get().GetExtensionForTag(CurrentData.TagName))
+			{
+				Extension->DisableDrawingForData(FLogVisualizer::Get().GetWorld(), Canvas, NULL, CurrentData.TagName, CurrentData, SelectedEntry.TimeStamp);
+			}
 		}
+		CachedDataBlocks.Reset();
 
-		if (!bIsValidByFilter)
+		for (const auto& CurrentData : SelectedEntry.DataBlocks)
 		{
-			Extension->DisableDrawingForData(FLogVisualizer::Get().GetWorld(), Canvas, NULL, TagName, CurrentData, SelectedEntry.TimeStamp);
+			const bool bIsValidByFilter = FCategoryFiltersManager::Get().MatchCategoryFilters(CurrentData.Category.ToString(), ELogVerbosity::All) && FCategoryFiltersManager::Get().MatchCategoryFilters(CurrentData.TagName.ToString(), ELogVerbosity::All);
+			if (bIsValidByFilter)
+			{
+				CachedDataBlocks.Add(CurrentData);
+			}
 		}
-		else
+	}
+
+	for (const auto CurrentData : CachedDataBlocks)
+	{
+		FVisualLogExtensionInterface* Extension = FVisualLogger::Get().GetExtensionForTag(CurrentData.TagName);
+		if (!Extension)
 		{
-			Extension->DrawData(FLogVisualizer::Get().GetWorld(), Canvas, NULL, TagName, CurrentData, SelectedEntry.TimeStamp);
+			continue;
 		}
+
+		Extension->DrawData(FLogVisualizer::Get().GetWorld(), Canvas, NULL, CurrentData.TagName, CurrentData, SelectedEntry.TimeStamp);
 	}
 
 	if (ULogVisualizerSessionSettings::StaticClass()->GetDefaultObject<ULogVisualizerSessionSettings>()->bEnableGraphsVisualization)
 	{
 		DrawHistogramGraphs(Canvas, NULL);
 	}
+
+	bDirtyData = false;
 }
 
 void FVisualLoggerCanvasRenderer::DrawHistogramGraphs(class UCanvas* Canvas, class APlayerController*)
@@ -144,73 +160,105 @@ void FVisualLoggerCanvasRenderer::DrawHistogramGraphs(class UCanvas* Canvas, cla
 	int32 LeftSideOutsideIndex = INDEX_NONE;
 	int32 RightSideOutsideIndex = INDEX_NONE;
 
-	for (int32 EntryIndex = 0; EntryIndex < ObjectItems.Num(); ++EntryIndex)
+	if (bDirtyData)
 	{
-		const FVisualLogEntry* CurrentEntry = &(ObjectItems[EntryIndex].Entry);
-		if (CurrentEntry->TimeStamp < TimeStampWindow.X)
+		CachedHistogramSamples.Reset();
+		TArray<FVisualLogHistogramSample> CurrentSamples;
+		for (int32 EntryIndex = 0; EntryIndex < ObjectItems.Num(); ++EntryIndex)
 		{
-			LeftSideOutsideIndex = EntryIndex;
-			continue;
-		}
+			const FVisualLogEntry* CurrentEntry = &(ObjectItems[EntryIndex].Entry);
+			if (CurrentEntry->TimeStamp < TimeStampWindow.X)
+			{
+				LeftSideOutsideIndex = EntryIndex;
+				continue;
+			}
 
-		if (CurrentEntry->TimeStamp > TimeStampWindow.Y)
-		{
-			RightSideOutsideIndex = EntryIndex;
-			break;
-		}
+			if (CurrentEntry->TimeStamp > TimeStampWindow.Y)
+			{
+				RightSideOutsideIndex = EntryIndex;
+				break;
+			}
 
-		const int32 SamplesNum = CurrentEntry->HistogramSamples.Num();
-		for (int32 SampleIndex = 0; SampleIndex < SamplesNum; ++SampleIndex)
-		{
-			FVisualLogHistogramSample CurrentSample = CurrentEntry->HistogramSamples[SampleIndex];
+			CurrentSamples.Reset();
+			const int32 SamplesNum = CurrentEntry->HistogramSamples.Num();
+			for (int32 SampleIndex = 0; SampleIndex < SamplesNum; ++SampleIndex)
+			{
+				const FVisualLogHistogramSample& CurrentSample = CurrentEntry->HistogramSamples[SampleIndex];
 
-			const FName CurrentCategory = CurrentSample.Category;
-			const FName CurrentGraphName = CurrentSample.GraphName;
-			const FName CurrentDataName = CurrentSample.DataName;
+				const FName CurrentCategory = CurrentSample.Category;
+				const FName CurrentGraphName = CurrentSample.GraphName;
+				const FName CurrentDataName = CurrentSample.DataName;
 
-			FString GraphFilterName = CurrentSample.GraphName.ToString() +TEXT("$") + CurrentSample.DataName.ToString();
-			const bool bIsValidByFilter = FCategoryFiltersManager::Get().MatchCategoryFilters(GraphFilterName, ELogVerbosity::All);
+				FString GraphFilterName = CurrentSample.GraphName.ToString() + TEXT("$") + CurrentSample.DataName.ToString();
+				const bool bIsValidByFilter = FCategoryFiltersManager::Get().MatchCategoryFilters(GraphFilterName, ELogVerbosity::All);
+				if (bIsValidByFilter)
+				{
+					CurrentSamples.Add(CurrentSample);
+				}
+			}
 
-			if (bIsValidByFilter)
+			if (CurrentSamples.Num() > 0)
 			{
-				FGraphData &GraphData = CollectedGraphs.FindOrAdd(CurrentSample.GraphName);
-				FGraphLineData &LineData = GraphData.GraphLines.FindOrAdd(CurrentSample.DataName);
-				LineData.DataName = CurrentSample.DataName;
-				LineData.Samples.Add(CurrentSample.SampleValue);
-
-				GraphData.Min.X = FMath::Min(GraphData.Min.X, CurrentSample.SampleValue.X);
-				GraphData.Min.Y = FMath::Min(GraphData.Min.Y, CurrentSample.SampleValue.Y);
-
-				GraphData.Max.X = FMath::Max(GraphData.Max.X, CurrentSample.SampleValue.X);
-				GraphData.Max.Y = FMath::Max(GraphData.Max.Y, CurrentSample.SampleValue.Y);
+				CachedHistogramSamples.Add(CurrentSamples);
 			}
 		}
 	}
-	
-	const int32 ExtremeValueIndexes[] = { LeftSideOutsideIndex != INDEX_NONE ? LeftSideOutsideIndex : 0, RightSideOutsideIndex != INDEX_NONE ? RightSideOutsideIndex : ObjectItems.Num() - 1 };
-	for (int32 ObjectIndex = 0; ObjectIndex < 2; ++ObjectIndex)
+
+	for (int32 EntryIndex = 0; EntryIndex < CachedHistogramSamples.Num(); ++EntryIndex)
 	{
-		const FVisualLogEntry* CurrentEntry = &(ObjectItems[ExtremeValueIndexes[ObjectIndex]].Entry);
-		const int32 SamplesNum = CurrentEntry->HistogramSamples.Num();
+		TArray<FVisualLogHistogramSample> &HistogramSamples = CachedHistogramSamples[EntryIndex];
+		const int32 SamplesNum = CachedHistogramSamples[EntryIndex].Num();
 		for (int32 SampleIndex = 0; SampleIndex < SamplesNum; ++SampleIndex)
 		{
-			FVisualLogHistogramSample CurrentSample = CurrentEntry->HistogramSamples[SampleIndex];
+			FVisualLogHistogramSample CurrentSample = HistogramSamples[SampleIndex];
 
-			const FName CurrentCategory = CurrentSample.Category;
-			const FName CurrentGraphName = CurrentSample.GraphName;
-			const FName CurrentDataName = CurrentSample.DataName;
+			FGraphData &GraphData = CollectedGraphs.FindOrAdd(CurrentSample.GraphName);
+			FGraphLineData &LineData = GraphData.GraphLines.FindOrAdd(CurrentSample.DataName);
+			LineData.DataName = CurrentSample.DataName;
+			LineData.Samples.Add(CurrentSample.SampleValue);
+			LineData.LeftExtreme = FVector2D::ZeroVector;
+			LineData.RightExtreme = FVector2D::ZeroVector;
 
-			FString GraphFilterName = CurrentSample.GraphName.ToString() + TEXT("_") + CurrentSample.DataName.ToString();
-			const bool bIsValidByFilter = FCategoryFiltersManager::Get().MatchCategoryFilters(GraphFilterName, ELogVerbosity::All);
-			if (bIsValidByFilter)
+			GraphData.Min.X = FMath::Min(GraphData.Min.X, CurrentSample.SampleValue.X);
+			GraphData.Min.Y = FMath::Min(GraphData.Min.Y, CurrentSample.SampleValue.Y);
+
+			GraphData.Max.X = FMath::Max(GraphData.Max.X, CurrentSample.SampleValue.X);
+			GraphData.Max.Y = FMath::Max(GraphData.Max.Y, CurrentSample.SampleValue.Y);
+		}
+	}
+	
+	if (bDirtyData)
+	{
+		const int32 ExtremeValueIndexes[] = { LeftSideOutsideIndex != INDEX_NONE ? LeftSideOutsideIndex : 0, RightSideOutsideIndex != INDEX_NONE ? RightSideOutsideIndex : ObjectItems.Num() - 1 };
+		for (int32 ObjectIndex = 0; ObjectIndex < 2; ++ObjectIndex)
+		{
+			const FVisualLogEntry* CurrentEntry = &(ObjectItems[ExtremeValueIndexes[ObjectIndex]].Entry);
+			const int32 SamplesNum = CurrentEntry->HistogramSamples.Num();
+			for (int32 SampleIndex = 0; SampleIndex < SamplesNum; ++SampleIndex)
 			{
-				FGraphData &GraphData = CollectedGraphs.FindOrAdd(CurrentSample.GraphName);
-				FGraphLineData &LineData = GraphData.GraphLines.FindOrAdd(CurrentSample.DataName);
-				LineData.DataName = CurrentSample.DataName;
-				if (ObjectIndex == 0)
-					LineData.LeftExtreme = CurrentSample.SampleValue;
-				else
-					LineData.RightExtreme = CurrentSample.SampleValue;
+				FVisualLogHistogramSample CurrentSample = CurrentEntry->HistogramSamples[SampleIndex];
+
+				const FName CurrentCategory = CurrentSample.Category;
+				const FName CurrentGraphName = CurrentSample.GraphName;
+				const FName CurrentDataName = CurrentSample.DataName;
+
+				//FString GraphFilterName = CurrentSample.GraphName.ToString() + TEXT("_") + CurrentSample.DataName.ToString();
+				//const bool bIsValidByFilter = FCategoryFiltersManager::Get().MatchCategoryFilters(GraphFilterName, ELogVerbosity::All);
+				//if (bIsValidByFilter)
+				{
+					FGraphData *GraphData = CollectedGraphs.Find(CurrentSample.GraphName);
+					if (GraphData)
+					{
+						FGraphLineData *LineData = GraphData->GraphLines.Find(CurrentSample.DataName);
+						if (LineData)
+						{
+							if (ObjectIndex == 0)
+								LineData->LeftExtreme = CurrentSample.SampleValue;
+							else
+								LineData->RightExtreme = CurrentSample.SampleValue;
+						}
+					}
+				}
 			}
 		}
 	}
diff --git a/Engine/Source/Developer/LogVisualizer/Private/VisualLoggerCanvasRenderer.h b/Engine/Source/Developer/LogVisualizer/Private/VisualLoggerCanvasRenderer.h
--- a/Engine/Source/Developer/LogVisualizer/Private/VisualLoggerCanvasRenderer.h
+++ b/Engine/Source/Developer/LogVisualizer/Private/VisualLoggerCanvasRenderer.h
@@ -3,18 +3,24 @@
 
 struct FVisualLoggerCanvasRenderer
 {
-	FVisualLoggerCanvasRenderer() 
+public:
+	FVisualLoggerCanvasRenderer() : bDirtyData(true)
 	{}
 
 	void DrawOnCanvas(class UCanvas* Canvas, class APlayerController*);
 	void OnItemSelectionChanged(const struct FVisualLogEntry& EntryItem);
 	void ObjectSelectionChanged(TSharedPtr<class STimeline> TimeLine);
+	void DirtyCachedData() { bDirtyData = true; }
 
 protected:
 	void DrawHistogramGraphs(class UCanvas* Canvas, class APlayerController* );
 
 private:
+	bool bDirtyData;
 	FVisualLogEntry SelectedEntry;
 	TMap<FString, TArray<FString> > UsedGraphCategories;
 	TWeakPtr<class STimeline> CurrentTimeLine;
+
+	TArray<FVisualLogDataBlock> CachedDataBlocks;
+	TArray<TArray<FVisualLogHistogramSample> > CachedHistogramSamples;
 };
