Added tooltip to fix issue that "user is not informed that only the selected rows in Visual Loggger will be saved".

diff --git a/Engine/Source/Developer/LogVisualizer/Private/SVisualLoggerToolbar.cpp b/Engine/Source/Developer/LogVisualizer/Private/SVisualLoggerToolbar.cpp
--- a/Engine/Source/Developer/LogVisualizer/Private/SVisualLoggerToolbar.cpp
+++ b/Engine/Source/Developer/LogVisualizer/Private/SVisualLoggerToolbar.cpp
@@ -28,7 +28,7 @@ TSharedRef<SWidget> SVisualLoggerToolbar::MakeToolbar(const TSharedRef<FUIComman
 
 		ToolBarBuilder.AddSeparator();
 		ToolBarBuilder.AddToolBarButton(FVisualLoggerCommands::Get().LoadFromVLog, NAME_None, LOCTEXT("Load", "Load"), TAttribute<FText>(), FSlateIcon(FLogVisualizerStyle::Get().GetStyleSetName(), TEXT("Toolbar.Load")));
-		ToolBarBuilder.AddToolBarButton(FVisualLoggerCommands::Get().SaveToVLog, NAME_None, LOCTEXT("Save", "Save"), TAttribute<FText>(), FSlateIcon(FLogVisualizerStyle::Get().GetStyleSetName(), TEXT("Toolbar.Save")));
+		ToolBarBuilder.AddToolBarButton(FVisualLoggerCommands::Get().SaveToVLog, NAME_None, LOCTEXT("SaveLogs", "Save"), LOCTEXT("SaveLogsTooltip", "Save selected logs/rows to file."), FSlateIcon(FLogVisualizerStyle::Get().GetStyleSetName(), TEXT("Toolbar.Save")));
 
 		ToolBarBuilder.AddSeparator();
 		ToolBarBuilder.AddToolBarButton(FVisualLoggerCommands::Get().FreeCamera, NAME_None, LOCTEXT("FreeCamera", "Camera"), TAttribute<FText>(), FSlateIcon(FLogVisualizerStyle::Get().GetStyleSetName(), TEXT("Toolbar.Camera")));
