Added "Save All" button to Visual Logger.

diff --git a/Engine/Source/Developer/LogVisualizer/Private/LogVisualizerStyle.cpp b/Engine/Source/Developer/LogVisualizer/Private/LogVisualizerStyle.cpp
--- a/Engine/Source/Developer/LogVisualizer/Private/LogVisualizerStyle.cpp
+++ b/Engine/Source/Developer/LogVisualizer/Private/LogVisualizerStyle.cpp
@@ -79,6 +79,7 @@ TSharedRef< FSlateStyleSet > FLogVisualizerStyle::Create()
 		Style.Set("Toolbar.Stop", new IMAGE_BRUSH("Icons/LV_Stop", Icon40x40));
 		Style.Set("Toolbar.Camera", new IMAGE_BRUSH("Icons/LV_Camera", Icon40x40));
 		Style.Set("Toolbar.Save", new IMAGE_BRUSH("Icons/LV_Save", Icon40x40));
+		Style.Set("Toolbar.SaveAll", new IMAGE_BRUSH("Icons/icon_file_saveall_40x", Icon40x40));
 		Style.Set("Toolbar.Load", new IMAGE_BRUSH("Icons/LV_Load", Icon40x40));
 		Style.Set("Toolbar.Remove", new IMAGE_BRUSH("Icons/Edit/icon_Edit_Delete_40x", Icon40x40));
 		Style.Set("Toolbar.Graphs", new IMAGE_BRUSH("Icons/icon_MatEd_Stats_40x", Icon40x40));
@@ -90,6 +91,7 @@ TSharedRef< FSlateStyleSet > FLogVisualizerStyle::Create()
 		Style.Set("Toolbar.Stop.small", new IMAGE_BRUSH("Icons/LV_Stop", Icon20x20));
 		Style.Set("Toolbar.Camera.small", new IMAGE_BRUSH("Icons/LV_Camera", Icon20x20));
 		Style.Set("Toolbar.Save.small", new IMAGE_BRUSH("Icons/LV_Save", Icon20x20));
+		Style.Set("Toolbar.SaveAll.small", new IMAGE_BRUSH("Icons/icon_file_saveall_40x", Icon20x20));
 		Style.Set("Toolbar.Load.small", new IMAGE_BRUSH("Icons/LV_Load", Icon20x20));
 		Style.Set("Toolbar.Remove.small", new IMAGE_BRUSH("Icons/Edit/icon_Edit_Delete_40x", Icon20x20));
 		Style.Set("Toolbar.Graphs.small", new IMAGE_BRUSH("Icons/icon_MatEd_Stats_40x", Icon20x20));
diff --git a/Engine/Source/Developer/LogVisualizer/Private/SVisualLogger.cpp b/Engine/Source/Developer/LogVisualizer/Private/SVisualLogger.cpp
--- a/Engine/Source/Developer/LogVisualizer/Private/SVisualLogger.cpp
+++ b/Engine/Source/Developer/LogVisualizer/Private/SVisualLogger.cpp
@@ -145,7 +145,8 @@ void SVisualLogger::Construct(const FArguments& InArgs, const TSharedRef<SDockTa
 	ActionList.MapAction(Commands.Resume, FExecuteAction::CreateRaw(this, &SVisualLogger::HandleResumeCommandExecute), FCanExecuteAction::CreateRaw(this, &SVisualLogger::HandleResumeCommandCanExecute), FIsActionChecked(), FIsActionButtonVisible::CreateRaw(this, &SVisualLogger::HandleResumeCommandIsVisible));
 	ActionList.MapAction(Commands.LoadFromVLog, FExecuteAction::CreateRaw(this, &SVisualLogger::HandleLoadCommandExecute), FCanExecuteAction::CreateRaw(this, &SVisualLogger::HandleLoadCommandCanExecute), FIsActionChecked(), FIsActionButtonVisible::CreateRaw(this, &SVisualLogger::HandleLoadCommandCanExecute));
 	ActionList.MapAction(Commands.SaveToVLog, FExecuteAction::CreateRaw(this, &SVisualLogger::HandleSaveCommandExecute), FCanExecuteAction::CreateRaw(this, &SVisualLogger::HandleSaveCommandCanExecute), FIsActionChecked(), FIsActionButtonVisible::CreateRaw(this, &SVisualLogger::HandleSaveCommandCanExecute));
-	ActionList.MapAction(Commands.FreeCamera, 
+	ActionList.MapAction(Commands.SaveAllToVLog, FExecuteAction::CreateRaw(this, &SVisualLogger::HandleSaveAllCommandExecute), FCanExecuteAction::CreateRaw(this, &SVisualLogger::HandleSaveCommandCanExecute), FIsActionChecked(), FIsActionButtonVisible::CreateRaw(this, &SVisualLogger::HandleSaveCommandCanExecute));
+	ActionList.MapAction(Commands.FreeCamera,
 		FExecuteAction::CreateRaw(this, &SVisualLogger::HandleCameraCommandExecute), 
 		FCanExecuteAction::CreateRaw(this, &SVisualLogger::HandleCameraCommandCanExecute), 
 		FIsActionChecked::CreateRaw(this, &SVisualLogger::HandleCameraCommandIsChecked),
@@ -328,6 +329,7 @@ void SVisualLogger::FillFileMenu(FMenuBuilder& MenuBuilder, const TSharedPtr<FTa
 	{
 		MenuBuilder.AddMenuEntry(FVisualLoggerCommands::Get().LoadFromVLog);
 		MenuBuilder.AddMenuEntry(FVisualLoggerCommands::Get().SaveToVLog);
+		MenuBuilder.AddMenuEntry(FVisualLoggerCommands::Get().SaveAllToVLog);
 	}
 	MenuBuilder.EndSection();
 	MenuBuilder.BeginSection("LogFilters", LOCTEXT("FIlterMenu", "Log Filters"));
@@ -620,10 +622,20 @@ bool SVisualLogger::HandleSaveCommandCanExecute() const
 	return OutTimelines.Num() > 0;
 }
 
+void SVisualLogger::HandleSaveAllCommandExecute()
+{
+	HandleSaveCommand(true);
+}
+
 void SVisualLogger::HandleSaveCommandExecute()
+{
+	HandleSaveCommand(false);
+}
+
+void SVisualLogger::HandleSaveCommand(bool bSaveAllData)
 {
 	TArray<TSharedPtr<class STimeline> > OutTimelines;
-	MainView->GetTimelines(OutTimelines, true);
+	MainView->GetTimelines(OutTimelines, !bSaveAllData);
 	if (OutTimelines.Num() == 0)
 	{
 		MainView->GetTimelines(OutTimelines);
diff --git a/Engine/Source/Developer/LogVisualizer/Private/SVisualLogger.h b/Engine/Source/Developer/LogVisualizer/Private/SVisualLogger.h
--- a/Engine/Source/Developer/LogVisualizer/Private/SVisualLogger.h
+++ b/Engine/Source/Developer/LogVisualizer/Private/SVisualLogger.h
@@ -92,6 +92,10 @@ public:
 	bool HandleSaveCommandCanExecute() const;
 	void HandleSaveCommandExecute();
 
+	void HandleSaveAllCommandExecute();
+
+	void HandleSaveCommand(bool bSaveAllData);
+
 	bool HandleCameraCommandCanExecute() const;
 	void HandleCameraCommandExecute();
 	bool HandleCameraCommandIsChecked() const;
diff --git a/Engine/Source/Developer/LogVisualizer/Private/SVisualLoggerToolbar.cpp b/Engine/Source/Developer/LogVisualizer/Private/SVisualLoggerToolbar.cpp
--- a/Engine/Source/Developer/LogVisualizer/Private/SVisualLoggerToolbar.cpp
+++ b/Engine/Source/Developer/LogVisualizer/Private/SVisualLoggerToolbar.cpp
@@ -29,6 +29,7 @@ TSharedRef<SWidget> SVisualLoggerToolbar::MakeToolbar(const TSharedRef<FUIComman
 		ToolBarBuilder.AddSeparator();
 		ToolBarBuilder.AddToolBarButton(FVisualLoggerCommands::Get().LoadFromVLog, NAME_None, LOCTEXT("Load", "Load"), TAttribute<FText>(), FSlateIcon(FLogVisualizerStyle::Get().GetStyleSetName(), TEXT("Toolbar.Load")));
 		ToolBarBuilder.AddToolBarButton(FVisualLoggerCommands::Get().SaveToVLog, NAME_None, LOCTEXT("SaveLogs", "Save"), LOCTEXT("SaveLogsTooltip", "Save selected logs/rows to file."), FSlateIcon(FLogVisualizerStyle::Get().GetStyleSetName(), TEXT("Toolbar.Save")));
+		ToolBarBuilder.AddToolBarButton(FVisualLoggerCommands::Get().SaveAllToVLog, NAME_None, LOCTEXT("SaveAllLogs", "Save All"), TAttribute<FText>(), FSlateIcon(FLogVisualizerStyle::Get().GetStyleSetName(), TEXT("Toolbar.SaveAll")));
 
 		ToolBarBuilder.AddSeparator();
 		ToolBarBuilder.AddToolBarButton(FVisualLoggerCommands::Get().FreeCamera, NAME_None, LOCTEXT("FreeCamera", "Camera"), TAttribute<FText>(), FSlateIcon(FLogVisualizerStyle::Get().GetStyleSetName(), TEXT("Toolbar.Camera")));
diff --git a/Engine/Source/Developer/LogVisualizer/Private/VisualLoggerCommands.h b/Engine/Source/Developer/LogVisualizer/Private/VisualLoggerCommands.h
--- a/Engine/Source/Developer/LogVisualizer/Private/VisualLoggerCommands.h
+++ b/Engine/Source/Developer/LogVisualizer/Private/VisualLoggerCommands.h
@@ -27,6 +27,7 @@ public:
 		UI_COMMAND(FreeCamera, "FreeCamera", "Enable free camera", EUserInterfaceActionType::ToggleButton, FInputChord());
 		UI_COMMAND(LoadFromVLog, "Load", "Load external vlogs", EUserInterfaceActionType::Button, FInputChord());
 		UI_COMMAND(SaveToVLog, "Save", "Save selected data to vlog file", EUserInterfaceActionType::Button, FInputChord());
+		UI_COMMAND(SaveAllToVLog, "SaveAll", "Save all to vlog file", EUserInterfaceActionType::Button, FInputChord());
 		UI_COMMAND(ToggleGraphs, "ToggleGraphs", "Toggle graphs visualization on/off", EUserInterfaceActionType::ToggleButton, FInputChord());
 		UI_COMMAND(ResetData, "Clear", "Clear all data", EUserInterfaceActionType::Button, FInputChord());
 		UI_COMMAND(HideEmpty, "HideEmpty", "Hide logs without text information", EUserInterfaceActionType::ToggleButton, FInputChord());
@@ -44,6 +45,7 @@ public:
 	TSharedPtr<FUICommandInfo> FreeCamera;
 	TSharedPtr<FUICommandInfo> LoadFromVLog;
 	TSharedPtr<FUICommandInfo> SaveToVLog;
+	TSharedPtr<FUICommandInfo> SaveAllToVLog;
 	TSharedPtr<FUICommandInfo> ToggleGraphs;
 	TSharedPtr<FUICommandInfo> ResetData;
 	TSharedPtr<FUICommandInfo> HideEmpty;
