Forgot about one change for CL #2570444.

diff --git a/Engine/Source/Developer/LogVisualizer/Private/SVisualLogger.cpp b/Engine/Source/Developer/LogVisualizer/Private/SVisualLogger.cpp
--- a/Engine/Source/Developer/LogVisualizer/Private/SVisualLogger.cpp
+++ b/Engine/Source/Developer/LogVisualizer/Private/SVisualLogger.cpp
@@ -449,7 +449,7 @@ void SVisualLogger::HandleStopRecordingCommandExecute()
 {
 	UWorld* World = FLogVisualizer::Get().GetWorld();
 
-	if (FParse::Param(FCommandLine::Get(), TEXT("DumpNavOctree")) == false && ULogVisualizerSettings::StaticClass()->GetDefaultObject<ULogVisualizerSettings>()->bLogNavOctreeOnStop)
+	if (FParse::Param(FCommandLine::Get(), TEXT("LogNavOctree")) == false && ULogVisualizerSettings::StaticClass()->GetDefaultObject<ULogVisualizerSettings>()->bLogNavOctreeOnStop)
 	{
 		FVisualLogger::NavigationDataDump(World, LogNavigation, ELogVerbosity::Log, INDEX_NONE, FBox());
 	}
