Added multiple logs selection to Visual Logger tool. It should partially fill UE-16336  task.

diff --git a/Engine/Source/Developer/LogVisualizer/Private/STimelinesContainer.cpp b/Engine/Source/Developer/LogVisualizer/Private/STimelinesContainer.cpp
--- a/Engine/Source/Developer/LogVisualizer/Private/STimelinesContainer.cpp
+++ b/Engine/Source/Developer/LogVisualizer/Private/STimelinesContainer.cpp
@@ -78,6 +78,7 @@ void STimelinesContainer::SetSelectionState(TSharedPtr<class STimeline> Affected
 		// Not selecting so remove the node from the selection set
 		SelectedNodes.Remove(AffectedNode);
 		AffectedNode->OnDeselect();
+		FLogVisualizer::Get().GetVisualLoggerEvents().OnObjectSelectionChanged.ExecuteIfBound(AffectedNode);
 	}
 }
 
diff --git a/Engine/Source/Developer/LogVisualizer/Private/SVisualLogger.cpp b/Engine/Source/Developer/LogVisualizer/Private/SVisualLogger.cpp
--- a/Engine/Source/Developer/LogVisualizer/Private/SVisualLogger.cpp
+++ b/Engine/Source/Developer/LogVisualizer/Private/SVisualLogger.cpp
@@ -731,7 +731,8 @@ void SVisualLogger::ResetData()
 	if (HelperActor)
 	{
 		HelperActor->OnItemSelectionChanged(FVisualLogDevice::FVisualLogEntryItem());
-		HelperActor->ObjectSelectionChanged(NULL);
+		TArray<TSharedPtr<class STimeline> > TimeLines;
+		HelperActor->ObjectSelectionChanged(TimeLines);
 	}
 
 	bGotHistogramData = false;
@@ -808,14 +809,21 @@ void SVisualLogger::OnFiltersChanged()
 
 void SVisualLogger::OnObjectSelectionChanged(TSharedPtr<class STimeline> TimeLine)
 {
+	TArray<TSharedPtr<class STimeline> > TimeLines;
+	GetTimelines(TimeLines, true);
+
+	LogsList->ObjectSelectionChanged(TimeLines);
+
 	VisualLoggerCanvasRenderer->ObjectSelectionChanged(TimeLine);
 	AVisualLoggerRenderingActor* HelperActor = Cast<AVisualLoggerRenderingActor>(FLogVisualizer::Get().GetVisualLoggerHelperActor());
 	if (HelperActor)
 	{
-		HelperActor->ObjectSelectionChanged(TimeLine);
+		HelperActor->ObjectSelectionChanged(TimeLines);
 	}
 	MainView->OnObjectSelectionChanged(TimeLine);
 	FLogVisualizer::Get().OnObjectSelectionChanged(TimeLine);
+
+	StatusView->ObjectSelectionChanged(TimeLines);
 }
 
 void SVisualLogger::OnFiltersSearchChanged(const FText& Filter)
diff --git a/Engine/Source/Developer/LogVisualizer/Private/SVisualLoggerLogsList.cpp b/Engine/Source/Developer/LogVisualizer/Private/SVisualLoggerLogsList.cpp
--- a/Engine/Source/Developer/LogVisualizer/Private/SVisualLoggerLogsList.cpp
+++ b/Engine/Source/Developer/LogVisualizer/Private/SVisualLoggerLogsList.cpp
@@ -53,7 +53,11 @@ FReply SVisualLoggerLogsList::OnKeyDown(const FGeometry& MyGeometry, const FKeyE
 		FString ClipboardString;
 		for (const TSharedPtr<struct FLogEntryItem>& CurrentItem : LogsLinesWidget->GetSelectedItems())
 		{
-			ClipboardString += CurrentItem->Category + FString(TEXT(" (")) + FString(FOutputDevice::VerbosityToString(CurrentItem->Verbosity)) + FString(TEXT(") ")) + CurrentItem->Line;
+			if (CurrentItem->Category.Len() > 0)
+				ClipboardString += CurrentItem->Category + FString(TEXT(" (")) + FString(FOutputDevice::VerbosityToString(CurrentItem->Verbosity)) + FString(TEXT(") ")) + CurrentItem->Line;
+			else
+				ClipboardString += CurrentItem->Line;
+
 			ClipboardString += TEXT("\n");
 		}
 		FPlatformMisc::ClipboardCopy(*ClipboardString);
@@ -69,50 +73,59 @@ void SVisualLoggerLogsList::OnFiltersChanged()
 	OnItemSelectionChanged(LogEntry);
 }
 
-void SVisualLoggerLogsList::LogEntryLineSelectionChanged(TSharedPtr<FLogEntryItem> SelectedItem, ESelectInfo::Type SelectInfo)
+TSharedRef<ITableRow> SVisualLoggerLogsList::LogEntryLinesGenerateRow(TSharedPtr<FLogEntryItem> Item, const TSharedRef<STableViewBase>& OwnerTable)
 {
-	if (SelectedItem.IsValid() == true)
+	if (Item->Category.Len() > 0)
 	{
-		FLogVisualizer::Get().GetVisualLoggerEvents().OnLogLineSelectionChanged.ExecuteIfBound(SelectedItem, SelectedItem->UserData, SelectedItem->TagName);
+		return SNew(STableRow< TSharedPtr<FString> >, OwnerTable)
+			[
+				SNew(SHorizontalBox)
+				+ SHorizontalBox::Slot()
+				.AutoWidth()
+				.Padding(FMargin(5.0f, 0.0f))
+				[
+					SNew(STextBlock)
+					.ColorAndOpacity(FSlateColor(Item->CategoryColor))
+					.Text(FText::FromString(Item->Category))
+					.HighlightText(this, &SVisualLoggerLogsList::GetFilterText)
+				]
+				+ SHorizontalBox::Slot()
+					.AutoWidth()
+					.Padding(FMargin(5.0f, 0.0f))
+					[
+						SNew(STextBlock)
+						.ColorAndOpacity(FSlateColor(Item->Verbosity == ELogVerbosity::Error ? FLinearColor::Red : (Item->Verbosity == ELogVerbosity::Warning ? FLinearColor::Yellow : FLinearColor::Gray)))
+						.Text(FText::FromString(FString(TEXT("(")) + FString(FOutputDevice::VerbosityToString(Item->Verbosity)) + FString(TEXT(")"))))
+					]
+				+ SHorizontalBox::Slot()
+					.Padding(FMargin(5.0f, 0.0f))
+					[
+						SNew(STextBlock)
+						.AutoWrapText(true)
+						.ColorAndOpacity(FSlateColor(Item->Verbosity == ELogVerbosity::Error ? FLinearColor::Red : (Item->Verbosity == ELogVerbosity::Warning ? FLinearColor::Yellow : FLinearColor::Gray)))
+						.Text(FText::FromString(Item->Line))
+						.HighlightText(this, &SVisualLoggerLogsList::GetFilterText)
+					]
+			];
+
 	}
 	else
 	{
-		FLogVisualizer::Get().GetVisualLoggerEvents().OnLogLineSelectionChanged.ExecuteIfBound(SelectedItem, 0, NAME_None);
-	}
-}
-
-TSharedRef<ITableRow> SVisualLoggerLogsList::LogEntryLinesGenerateRow(TSharedPtr<FLogEntryItem> Item, const TSharedRef<STableViewBase>& OwnerTable)
-{
-	return SNew(STableRow< TSharedPtr<FString> >, OwnerTable)
-		[
-			SNew(SHorizontalBox)
-			+ SHorizontalBox::Slot()
-			.AutoWidth()
-			.Padding(FMargin(5.0f, 0.0f))
-			[
-				SNew(STextBlock)
-				.ColorAndOpacity(FSlateColor(Item->CategoryColor))
-				.Text( FText::FromString(Item->Category) )
-				.HighlightText(this, &SVisualLoggerLogsList::GetFilterText)
-			]
-			+ SHorizontalBox::Slot()
-			.AutoWidth()
-			.Padding(FMargin(5.0f, 0.0f))
-			[
-				SNew(STextBlock)
-				.ColorAndOpacity(FSlateColor(Item->Verbosity == ELogVerbosity::Error ? FLinearColor::Red : (Item->Verbosity == ELogVerbosity::Warning ? FLinearColor::Yellow : FLinearColor::Gray)))
-				.Text(FText::FromString(FString(TEXT("(")) + FString(FOutputDevice::VerbosityToString(Item->Verbosity)) + FString(TEXT(")"))))
-			]
-			+ SHorizontalBox::Slot()
-			.Padding(FMargin(5.0f, 0.0f))
+		return SNew(STableRow< TSharedPtr<FString> >, OwnerTable)
 			[
-				SNew(STextBlock)
-				.AutoWrapText(true)
-				.ColorAndOpacity(FSlateColor(Item->Verbosity == ELogVerbosity::Error ? FLinearColor::Red : (Item->Verbosity == ELogVerbosity::Warning ? FLinearColor::Yellow : FLinearColor::Gray)))
-				.Text(FText::FromString(Item->Line))
-				.HighlightText(this, &SVisualLoggerLogsList::GetFilterText)
-			]
-		];
+				SNew(SHorizontalBox)
+				+ SHorizontalBox::Slot()
+				.Padding(FMargin(5.0f, 0.0f))
+				[
+					SNew(STextBlock)
+					/*.AutoWrapText(true)*/
+					.ColorAndOpacity(FSlateColor(FLinearColor::White))
+					.Text(FText::FromString(Item->Line))
+					.HighlightText(this, &SVisualLoggerLogsList::GetFilterText)
+					.Justification(ETextJustify::Center)
+				]
+			];
+	}
 }
 
 FText SVisualLoggerLogsList::GetFilterText() const
@@ -127,11 +140,79 @@ void SVisualLoggerLogsList::OnFiltersSearchChanged(const FText& Filter)
 	OnFiltersChanged();
 }
 
+void SVisualLoggerLogsList::LogEntryLineSelectionChanged(TSharedPtr<FLogEntryItem> SelectedItem, ESelectInfo::Type SelectInfo)
+{
+	if (SelectedItem.IsValid() == true)
+	{
+		FLogVisualizer::Get().GetVisualLoggerEvents().OnLogLineSelectionChanged.ExecuteIfBound(SelectedItem, SelectedItem->UserData, SelectedItem->TagName);
+	}
+	else
+	{
+		FLogVisualizer::Get().GetVisualLoggerEvents().OnLogLineSelectionChanged.ExecuteIfBound(SelectedItem, 0, NAME_None);
+	}
+}
+
+void SVisualLoggerLogsList::ObjectSelectionChanged(TArray<TSharedPtr<class STimeline> >& TimeLines)
+{
+	CurrentLogEntry.Entry.Reset();
+	LogEntryLines.Reset();
+	LogsLinesWidget->RequestListRefresh();
+
+	SelectedTimeLines = TimeLines;
+}
+
 void SVisualLoggerLogsList::OnItemSelectionChanged(const FVisualLogDevice::FVisualLogEntryItem& LogEntry)
 {
 	CurrentLogEntry.Entry.Reset();
 	LogEntryLines.Reset();
 
+	GenerateLogs(LogEntry, SelectedTimeLines.Num() > 1);
+
+	if (SelectedTimeLines.Num() > 1)
+	{
+		for (TSharedPtr<class STimeline> CurrentTimeline : SelectedTimeLines)
+		{
+			if (CurrentTimeline.IsValid() == false || CurrentTimeline->GetVisibility().IsVisible() == false)
+			{
+				continue;
+			}
+
+			if (LogEntry.OwnerName == CurrentTimeline->GetName())
+			{
+				continue;
+			}
+
+			const TArray<FVisualLogDevice::FVisualLogEntryItem>& Entries = CurrentTimeline->GetEntries();
+
+			int32 BestItemIndex = INDEX_NONE;
+			float BestDistance = MAX_FLT;
+			for (int32 Index = 0; Index < Entries.Num(); Index++)
+			{
+				auto& CurrentEntryItem = Entries[Index];
+
+				if (CurrentTimeline->IsEntryHidden(CurrentEntryItem))
+				{
+					continue;
+				}
+				TArray<FVisualLoggerCategoryVerbosityPair> OutCategories;
+				const float CurrentDist = FMath::Abs(CurrentEntryItem.Entry.TimeStamp - LogEntry.Entry.TimeStamp);
+				if (CurrentDist < BestDistance)
+				{
+					BestDistance = CurrentDist;
+					BestItemIndex = Index;
+				}
+			}
+
+			if (Entries.IsValidIndex(BestItemIndex))
+			{
+				GenerateLogs(Entries[BestItemIndex], true);
+			}
+		}
+	}
+}
+
+void SVisualLoggerLogsList::GenerateLogs(const FVisualLogDevice::FVisualLogEntryItem& LogEntry, bool bGenerateHeader)
+{
 	TArray<FVisualLoggerCategoryVerbosityPair> OutCategories;
 	FVisualLoggerHelpers::GetCategories(LogEntry.Entry, OutCategories);
 	bool bHasValidCategory = false;
@@ -145,6 +226,18 @@ void SVisualLoggerLogsList::OnItemSelectionChanged(const FVisualLogDevice::FVisu
 		return;
 	}
 
+	if (bGenerateHeader)
+	{
+		FLogEntryItem EntryItem;
+		EntryItem.Category = FString();
+		EntryItem.CategoryColor = FLinearColor::Black;
+		EntryItem.Verbosity = ELogVerbosity::VeryVerbose;
+		EntryItem.Line = LogEntry.OwnerName.ToString();
+		EntryItem.UserData = 0;
+		EntryItem.TagName = NAME_None;
+		LogEntryLines.Add(MakeShareable(new FLogEntryItem(EntryItem)));
+	}
+
 	CurrentLogEntry = LogEntry;
 	const FVisualLogLine* LogLine = LogEntry.Entry.LogLines.GetData();
 	const bool bSearchInsideLogs = ULogVisualizerSettings::StaticClass()->GetDefaultObject<ULogVisualizerSettings>()->bSearchInsideLogs;
diff --git a/Engine/Source/Developer/LogVisualizer/Private/SVisualLoggerLogsList.h b/Engine/Source/Developer/LogVisualizer/Private/SVisualLoggerLogsList.h
--- a/Engine/Source/Developer/LogVisualizer/Private/SVisualLoggerLogsList.h
+++ b/Engine/Source/Developer/LogVisualizer/Private/SVisualLoggerLogsList.h
@@ -14,8 +14,10 @@ public:
 	void LogEntryLineSelectionChanged(TSharedPtr<FLogEntryItem> SelectedItem, ESelectInfo::Type SelectInfo);
 	virtual bool SupportsKeyboardFocus() const override { return true; }
 	virtual FReply OnKeyDown(const FGeometry& MyGeometry, const FKeyEvent& InKeyEvent) override;
+	void ObjectSelectionChanged(TArray<TSharedPtr<class STimeline> >& TimeLines);
 
 	void OnItemSelectionChanged(const FVisualLogDevice::FVisualLogEntryItem& EntryItem);
+	void GenerateLogs(const FVisualLogDevice::FVisualLogEntryItem& EntryItem, bool bGenerateHeader);
 	void OnFiltersChanged();
 	void OnFiltersSearchChanged(const FText& Filter);
 	FText GetFilterText() const;
@@ -25,4 +27,5 @@ protected:
 	TSharedPtr<SListView<TSharedPtr<struct FLogEntryItem> > > LogsLinesWidget;
 	TArray<TSharedPtr<struct FLogEntryItem> > LogEntryLines;
 	FVisualLogDevice::FVisualLogEntryItem CurrentLogEntry;
+	TArray<TSharedPtr<class STimeline> > SelectedTimeLines;
 };
diff --git a/Engine/Source/Developer/LogVisualizer/Private/SVisualLoggerStatusView.cpp b/Engine/Source/Developer/LogVisualizer/Private/SVisualLoggerStatusView.cpp
--- a/Engine/Source/Developer/LogVisualizer/Private/SVisualLoggerStatusView.cpp
+++ b/Engine/Source/Developer/LogVisualizer/Private/SVisualLoggerStatusView.cpp
@@ -32,16 +32,37 @@ void SVisualLoggerStatusView::Construct(const FArguments& InArgs, const TSharedR
 			.Padding(1)
 			.BorderImage(FLogVisualizerStyle::Get().GetBrush("ToolPanel.GroupBorder"))
 			[
-				SAssignNew(StatusItemsView, STreeView<TSharedPtr<FLogStatusItem> >)
-				.ItemHeight(40.0f)
-				.TreeItemsSource(&StatusItems)
-				.OnGenerateRow(this, &SVisualLoggerStatusView::HandleGenerateLogStatus)
-				.OnGetChildren(this, &SVisualLoggerStatusView::OnLogStatusGetChildren)
-				.SelectionMode(ESelectionMode::None)
+				SNew(SVerticalBox)
+				+ SVerticalBox::Slot()
+				[
+					SAssignNew(StatusItemsView, STreeView<TSharedPtr<FLogStatusItem> >)
+					.ItemHeight(40.0f)
+					.TreeItemsSource(&StatusItems)
+					.OnGenerateRow(this, &SVisualLoggerStatusView::HandleGenerateLogStatus)
+					.OnGetChildren(this, &SVisualLoggerStatusView::OnLogStatusGetChildren)
+					.SelectionMode(ESelectionMode::None)
+					.Visibility(EVisibility::Visible)
+				]
+				+ SVerticalBox::Slot()
+				.FillHeight(1.0)
+				.VAlign(EVerticalAlignment::VAlign_Center)
+				[
+					SAssignNew(NotificationText, STextBlock)
+					.ColorAndOpacity(FSlateColor(FLinearColor::White))
+					.Text(FText::FromString(FString(TEXT("Multiple data selected"))))
+					.Justification(ETextJustify::Center)
+					.Visibility(EVisibility::Collapsed)
+				]
 			]
 		];
 }
 
+void SVisualLoggerStatusView::HideData(bool bInHide)
+{
+	NotificationText->SetVisibility(bInHide ? EVisibility::Visible : EVisibility::Collapsed);
+	StatusItemsView->SetVisibility(bInHide ? EVisibility::Collapsed : EVisibility::Visible);
+}
+
 void GenerateChildren(TSharedPtr<FLogStatusItem> StatusItem, const FVisualLogStatusCategory LogCategory)
 {
 	for (int32 LineIndex = 0; LineIndex < LogCategory.Data.Num(); LineIndex++)
@@ -61,15 +82,27 @@ void GenerateChildren(TSharedPtr<FLogStatusItem> StatusItem, const FVisualLogSta
 		GenerateChildren(ChildCategory, Child);
 	}
 }
+
+void SVisualLoggerStatusView::ObjectSelectionChanged(TArray<TSharedPtr<class STimeline> >& TimeLines)
+{
+	SelectedTimeLines = TimeLines;
+	StatusItems.Reset();
+	StatusItemsView->RequestTreeRefresh();
+}
+
+
 void SVisualLoggerStatusView::OnItemSelectionChanged(const FVisualLogDevice::FVisualLogEntryItem& LogEntry)
 {
+
 	TArray<FString> ExpandedCategories;
-	for (int32 ItemIndex = 0; ItemIndex < StatusItems.Num(); ItemIndex++)
 	{
-		const bool bIsExpanded = StatusItemsView->IsItemExpanded(StatusItems[ItemIndex]);
-		if (bIsExpanded)
+		for (int32 ItemIndex = 0; ItemIndex < StatusItems.Num(); ItemIndex++)
 		{
-			ExpandedCategories.Add(StatusItems[ItemIndex]->ItemText);
+			const bool bIsExpanded = StatusItemsView->IsItemExpanded(StatusItems[ItemIndex]);
+			if (bIsExpanded)
+			{
+				ExpandedCategories.Add(StatusItems[ItemIndex]->ItemText);
+			}
 		}
 	}
 
@@ -78,6 +111,71 @@ void SVisualLoggerStatusView::OnItemSelectionChanged(const FVisualLogDevice::FVi
 	FString TimestampDesc = FString::Printf(TEXT("%.2fs"), LogEntry.Entry.TimeStamp);
 	StatusItems.Add(MakeShareable(new FLogStatusItem(LOCTEXT("VisLogTimestamp", "Time").ToString(), TimestampDesc)));
 
+	GenerateStatusData(LogEntry, SelectedTimeLines.Num() > 1);
+
+	if (SelectedTimeLines.Num() > 1)
+	{
+		for (TSharedPtr<class STimeline> CurrentTimeline : SelectedTimeLines)
+		{
+			if (CurrentTimeline.IsValid() == false || CurrentTimeline->GetVisibility().IsVisible() == false)
+			{
+				continue;
+			}
+
+			if (LogEntry.OwnerName == CurrentTimeline->GetName())
+			{
+				continue;
+			}
+
+			const TArray<FVisualLogDevice::FVisualLogEntryItem>& Entries = CurrentTimeline->GetEntries();
+
+			int32 BestItemIndex = INDEX_NONE;
+			float BestDistance = MAX_FLT;
+			for (int32 Index = 0; Index < Entries.Num(); Index++)
+			{
+				auto& CurrentEntryItem = Entries[Index];
+
+				if (CurrentTimeline->IsEntryHidden(CurrentEntryItem))
+				{
+					continue;
+				}
+				TArray<FVisualLoggerCategoryVerbosityPair> OutCategories;
+				const float CurrentDist = FMath::Abs(CurrentEntryItem.Entry.TimeStamp - LogEntry.Entry.TimeStamp);
+				if (CurrentDist < BestDistance)
+				{
+					BestDistance = CurrentDist;
+					BestItemIndex = Index;
+				}
+			}
+
+			if (Entries.IsValidIndex(BestItemIndex))
+			{
+				GenerateStatusData(Entries[BestItemIndex], true);
+			}
+		}
+	}
+
+	{
+		for (int32 ItemIndex = 0; ItemIndex < StatusItems.Num(); ItemIndex++)
+		{
+			for (const FString& Category : ExpandedCategories)
+			{	if (StatusItems[ItemIndex]->ItemText == Category)
+				{
+					StatusItemsView->SetItemExpansion(StatusItems[ItemIndex], true);
+					break;
+				}
+			}
+		}
+	}
+}
+
+void SVisualLoggerStatusView::GenerateStatusData(const FVisualLogDevice::FVisualLogEntryItem& LogEntry, bool bAddHeader)
+{
+	if (bAddHeader)
+	{
+		StatusItems.Add(MakeShareable(new FLogStatusItem(LOCTEXT("VisLogOwnerName", "Owner Name").ToString(), LogEntry.OwnerName.ToString())));
+	}
+
 	for (int32 CategoryIndex = 0; CategoryIndex < LogEntry.Entry.Status.Num(); CategoryIndex++)
 	{
 		if (LogEntry.Entry.Status[CategoryIndex].Data.Num() <= 0 && LogEntry.Entry.Status[CategoryIndex].Children.Num() == 0)
@@ -92,7 +190,8 @@ void SVisualLoggerStatusView::OnItemSelectionChanged(const FVisualLogDevice::FVi
 			const bool bHasValue = LogEntry.Entry.Status[CategoryIndex].GetDesc(LineIndex, KeyDesc, ValueDesc);
 			if (bHasValue)
 			{
-				StatusItem->Children.Add(MakeShareable(new FLogStatusItem(KeyDesc, ValueDesc)));
+				TSharedPtr< FLogStatusItem > NewItem = MakeShareable(new FLogStatusItem(KeyDesc, ValueDesc));
+				StatusItem->Children.Add(NewItem);
 			}
 		}
 
@@ -107,18 +206,6 @@ void SVisualLoggerStatusView::OnItemSelectionChanged(const FVisualLogDevice::FVi
 	}
 
 	StatusItemsView->RequestTreeRefresh();
-
-	for (int32 ItemIndex = 0; ItemIndex < StatusItems.Num(); ItemIndex++)
-	{
-		for (const FString& Category : ExpandedCategories)
-		{
-			if (StatusItems[ItemIndex]->ItemText == Category)
-			{
-				StatusItemsView->SetItemExpansion(StatusItems[ItemIndex], true);
-				break;
-			}
-		}
-	}
 }
 
 void SVisualLoggerStatusView::OnLogStatusGetChildren(TSharedPtr<FLogStatusItem> InItem, TArray< TSharedPtr<FLogStatusItem> >& OutItems)
diff --git a/Engine/Source/Developer/LogVisualizer/Private/SVisualLoggerStatusView.h b/Engine/Source/Developer/LogVisualizer/Private/SVisualLoggerStatusView.h
--- a/Engine/Source/Developer/LogVisualizer/Private/SVisualLoggerStatusView.h
+++ b/Engine/Source/Developer/LogVisualizer/Private/SVisualLoggerStatusView.h
@@ -11,10 +11,16 @@ public:
 	void Construct(const FArguments& InArgs, const TSharedRef<FUICommandList>& InCommandList);
 
 	void OnItemSelectionChanged(const FVisualLogDevice::FVisualLogEntryItem&);
+	void GenerateStatusData(const FVisualLogDevice::FVisualLogEntryItem&, bool bAddHeader);
+	void ObjectSelectionChanged(TArray<TSharedPtr<class STimeline> >& TimeLines);
 
 	TSharedRef<ITableRow> HandleGenerateLogStatus(TSharedPtr<struct FLogStatusItem> InItem, const TSharedRef<STableViewBase>& OwnerTable);
 	void OnLogStatusGetChildren(TSharedPtr<struct FLogStatusItem> InItem, TArray< TSharedPtr<FLogStatusItem> >& OutItems);
+	void HideData(bool bInHide);
+
 protected:
 	TSharedPtr< STreeView< TSharedPtr<FLogStatusItem> > > StatusItemsView;
+	TSharedPtr< STextBlock > NotificationText;
 	TArray< TSharedPtr<FLogStatusItem> > StatusItems;
+	TArray<TSharedPtr<class STimeline> > SelectedTimeLines;
 };
diff --git a/Engine/Source/Developer/LogVisualizer/Private/VisualLoggerRenderingActor.cpp b/Engine/Source/Developer/LogVisualizer/Private/VisualLoggerRenderingActor.cpp
--- a/Engine/Source/Developer/LogVisualizer/Private/VisualLoggerRenderingActor.cpp
+++ b/Engine/Source/Developer/LogVisualizer/Private/VisualLoggerRenderingActor.cpp
@@ -59,15 +59,39 @@ FPrimitiveSceneProxy* UVisualLoggerRenderingComponent::CreateSceneProxy()
 	ULogVisualizerSettings *Settings = ULogVisualizerSettings::StaticClass()->GetDefaultObject<ULogVisualizerSettings>();
 	FVisualLoggerSceneProxy *VLogSceneProxy = new FVisualLoggerSceneProxy(this);
 	VLogSceneProxy->SolidMeshMaterial = Settings->GetDebugMeshMaterial();
-	VLogSceneProxy->Spheres = RenderingActor->Points;
-	VLogSceneProxy->Lines = RenderingActor->Lines;
-	VLogSceneProxy->Cones = RenderingActor->Cones;
-	VLogSceneProxy->Boxes = RenderingActor->Boxes;
-	VLogSceneProxy->Meshes = RenderingActor->Meshes;
-	VLogSceneProxy->Cones = RenderingActor->Cones;
-	VLogSceneProxy->Texts = RenderingActor->Texts;
-	VLogSceneProxy->Cylinders = RenderingActor->Cylinders;
-	VLogSceneProxy->Capsles = RenderingActor->Capsles;
+	VLogSceneProxy->Spheres = RenderingActor->PrimaryDebugShapes.Points;
+	VLogSceneProxy->Lines = RenderingActor->PrimaryDebugShapes.Lines;
+	VLogSceneProxy->Cones = RenderingActor->PrimaryDebugShapes.Cones;
+	VLogSceneProxy->Boxes = RenderingActor->PrimaryDebugShapes.Boxes;
+	VLogSceneProxy->Meshes = RenderingActor->PrimaryDebugShapes.Meshes;
+	VLogSceneProxy->Cones = RenderingActor->PrimaryDebugShapes.Cones;
+	VLogSceneProxy->Texts = RenderingActor->PrimaryDebugShapes.Texts;
+	VLogSceneProxy->Cylinders = RenderingActor->PrimaryDebugShapes.Cylinders;
+	VLogSceneProxy->Capsles = RenderingActor->PrimaryDebugShapes.Capsles;
+
+	{
+		VLogSceneProxy->Spheres.Append(RenderingActor->SecondaryDebugShapes.Points);
+		VLogSceneProxy->Lines.Append(RenderingActor->SecondaryDebugShapes.Lines);
+		VLogSceneProxy->Cones.Append(RenderingActor->SecondaryDebugShapes.Cones);
+		VLogSceneProxy->Boxes.Append(RenderingActor->SecondaryDebugShapes.Boxes);
+		VLogSceneProxy->Meshes.Append(RenderingActor->SecondaryDebugShapes.Meshes);
+		VLogSceneProxy->Cones.Append(RenderingActor->SecondaryDebugShapes.Cones);
+		VLogSceneProxy->Texts.Append(RenderingActor->SecondaryDebugShapes.Texts);
+		VLogSceneProxy->Cylinders.Append(RenderingActor->SecondaryDebugShapes.Cylinders);
+		VLogSceneProxy->Capsles.Append(RenderingActor->SecondaryDebugShapes.Capsles);
+	}
+
+	{
+		VLogSceneProxy->Spheres.Append(RenderingActor->TestDebugShapes.Points);
+		VLogSceneProxy->Lines.Append(RenderingActor->TestDebugShapes.Lines);
+		VLogSceneProxy->Cones.Append(RenderingActor->TestDebugShapes.Cones);
+		VLogSceneProxy->Boxes.Append(RenderingActor->TestDebugShapes.Boxes);
+		VLogSceneProxy->Meshes.Append(RenderingActor->TestDebugShapes.Meshes);
+		VLogSceneProxy->Cones.Append(RenderingActor->TestDebugShapes.Cones);
+		VLogSceneProxy->Texts.Append(RenderingActor->TestDebugShapes.Texts);
+		VLogSceneProxy->Cylinders.Append(RenderingActor->TestDebugShapes.Cylinders);
+		VLogSceneProxy->Capsles.Append(RenderingActor->TestDebugShapes.Capsles);
+	}
 
 	return VLogSceneProxy;
 }
@@ -119,10 +143,10 @@ void AVisualLoggerRenderingActor::AddDebugRendering()
 	{
 		const FVector BoxExtent(100, 100, 100);
 		const FBox Box(FVector(128), FVector(300));
-		Boxes.Add(FDebugRenderSceneProxy::FDebugBox(Box, FColor::Red));
+		TestDebugShapes.Boxes.Add(FDebugRenderSceneProxy::FDebugBox(Box, FColor::Red));
 		FTransform Trans;
 		Trans.SetRotation(FQuat::MakeFromEuler(FVector(0.1, 0.2, 1.2)));
-		Boxes.Add(FDebugRenderSceneProxy::FDebugBox(Box, FColor::Red, Trans));
+		TestDebugShapes.Boxes.Add(FDebugRenderSceneProxy::FDebugBox(Box, FColor::Red, Trans));
 	}
 	{
 		const FVector Orgin = FVector(400,0,128);
@@ -131,14 +155,14 @@ void AVisualLoggerRenderingActor::AddDebugRendering()
 
 		FVector YAxis, ZAxis;
 		Direction.FindBestAxisVectors(YAxis, ZAxis);
-		Cones.Add(FDebugRenderSceneProxy::FCone(FScaleMatrix(FVector(Length)) * FMatrix(Direction, YAxis, ZAxis, Orgin), 30, 30, FColor::Blue));
+		TestDebugShapes.Cones.Add(FDebugRenderSceneProxy::FCone(FScaleMatrix(FVector(Length)) * FMatrix(Direction, YAxis, ZAxis, Orgin), 30, 30, FColor::Blue));
 	}
 	{
 		const FVector Start = FVector(700, 0, 128);
 		const FVector End = FVector(700, 0, 128+300);
 		const float Radius = 200;
 		const float HalfHeight = 150;
-		Cylinders.Add(FDebugRenderSceneProxy::FWireCylinder(Start + FVector(0, 0, HalfHeight), Radius, HalfHeight, FColor::Magenta));
+		TestDebugShapes.Cylinders.Add(FDebugRenderSceneProxy::FWireCylinder(Start + FVector(0, 0, HalfHeight), Radius, HalfHeight, FColor::Magenta));
 	}
 
 	{
@@ -152,40 +176,31 @@ void AVisualLoggerRenderingActor::AddDebugRendering()
 		const FVector YAxis = Axes.GetScaledAxis(EAxis::Y);
 		const FVector ZAxis = Axes.GetScaledAxis(EAxis::Z);
 
-		Capsles.Add(FDebugRenderSceneProxy::FCapsule(Center, Radius, XAxis, YAxis, ZAxis, HalfHeight, FColor::Yellow));
+		TestDebugShapes.Capsles.Add(FDebugRenderSceneProxy::FCapsule(Center, Radius, XAxis, YAxis, ZAxis, HalfHeight, FColor::Yellow));
 	}
 	{
 		const float Radius = 50;
-		Points.Add(FDebugRenderSceneProxy::FSphere(10, FVector(1300, 0, 128), FColor::White));
+		TestDebugShapes.Points.Add(FDebugRenderSceneProxy::FSphere(10, FVector(1300, 0, 128), FColor::White));
 	}
 }
 
-void AVisualLoggerRenderingActor::ObjectSelectionChanged(TSharedPtr<class STimeline> TimeLine)
+void AVisualLoggerRenderingActor::ObjectSelectionChanged(TArray<TSharedPtr<class STimeline> >& TimeLines)
 {
-	if (TimeLine.IsValid() == false)
-	{
-		Lines.Reset();
-		Points.Reset();
-		Cones.Reset();
-		Boxes.Reset();
-		Meshes.Reset();
-		Texts.Reset();
-		Cylinders.Reset();
-		Capsles.Reset();
-		LogEntriesPath.Reset();
-		MarkComponentsRenderStateDirty();
-		return;
-	}
+	PrimaryDebugShapes.Reset();
 
-	const TArray<FVisualLogDevice::FVisualLogEntryItem>& Entries = TimeLine->GetEntries();
-	LogEntriesPath.Reset();
-	for (const auto &CurrentEntry : Entries)
+	if (TimeLines.Num() > 0)
 	{
-		if (CurrentEntry.Entry.Location != FVector::ZeroVector)
+		const TArray<FVisualLogDevice::FVisualLogEntryItem>& Entries = TimeLines[TimeLines.Num() - 1]->GetEntries();
+		for (const auto &CurrentEntry : Entries)
 		{
-			LogEntriesPath.Add(CurrentEntry.Entry.Location);
+			if (CurrentEntry.Entry.Location != FVector::ZeroVector)
+			{
+				PrimaryDebugShapes.LogEntriesPath.Add(CurrentEntry.Entry.Location);
+			}
 		}
 	}
+	SelectedTimeLines = TimeLines;
+	MarkComponentsRenderStateDirty();
 }
 
 namespace
@@ -221,21 +236,12 @@ namespace
 	}
 }
 
-void AVisualLoggerRenderingActor::OnItemSelectionChanged(const FVisualLogDevice::FVisualLogEntryItem& EntryItem)
+void AVisualLoggerRenderingActor::GetDebugShapes(const FVisualLogDevice::FVisualLogEntryItem& EntryItem, FTimelineDebugShapes& DebugShapes)
 {
 	const FVisualLogEntry* Entry = &EntryItem.Entry;
 	const FVisualLogShapeElement* ElementToDraw = Entry->ElementsToDraw.GetData();
 	const int32 ElementsCount = Entry->ElementsToDraw.Num();
 
-	Lines.Reset();
-	Points.Reset();
-	Cones.Reset();
-	Boxes.Reset();
-	Meshes.Reset();
-	Texts.Reset();
-	Cylinders.Reset();
-	Capsles.Reset();
-
 #if 0
 	AddDebugRendering();
 #endif
@@ -245,16 +251,16 @@ void AVisualLoggerRenderingActor::OnItemSelectionChanged(const FVisualLogDevice:
 		const FVector DirectionNorm = FVector(0, 0, 1).GetSafeNormal();
 		FVector YAxis, ZAxis;
 		DirectionNorm.FindBestAxisVectors(YAxis, ZAxis);
-		Cones.Add(FDebugRenderSceneProxy::FCone(FScaleMatrix(FVector(Length)) * FMatrix(DirectionNorm, YAxis, ZAxis, Entry->Location), 5, 5, FColor::Red));
+		DebugShapes.Cones.Add(FDebugRenderSceneProxy::FCone(FScaleMatrix(FVector(Length)) * FMatrix(DirectionNorm, YAxis, ZAxis, Entry->Location), 5, 5, FColor::Red));
 	}
 
-	if (LogEntriesPath.Num())
+	if (DebugShapes.LogEntriesPath.Num())
 	{
-		FVector Location = LogEntriesPath[0];
-		for (int32 Index = 1; Index < LogEntriesPath.Num(); ++Index)
+		FVector Location = DebugShapes.LogEntriesPath[0];
+		for (int32 Index = 1; Index < DebugShapes.LogEntriesPath.Num(); ++Index)
 		{
-			const FVector CurrentLocation = LogEntriesPath[Index];
-			Lines.Add(FDebugRenderSceneProxy::FDebugLine(Location, CurrentLocation, FColor(160, 160, 240), 2.0));
+			const FVector CurrentLocation = DebugShapes.LogEntriesPath[Index];
+			DebugShapes.Lines.Add(FDebugRenderSceneProxy::FDebugLine(Location, CurrentLocation, FColor(160, 160, 240), 2.0));
 			Location = CurrentLocation;
 		}
 	}
@@ -305,11 +311,11 @@ void AVisualLoggerRenderingActor::OnItemSelectionChanged(const FVisualLogDevice:
 			for (int32 Index = 0; Index < ElementToDraw->Points.Num(); ++Index)
 			{
 				const FVector& Point = ElementToDraw->Points[Index];
-				Points.Add(FDebugRenderSceneProxy::FSphere(Radius, Point, Color));
+				DebugShapes.Points.Add(FDebugRenderSceneProxy::FSphere(Radius, Point, Color));
 				if (bDrawLabel)
 				{
 					const FString PrintString = FString::Printf(TEXT("%s_%d"), *ElementToDraw->Description, Index);
-					Texts.Add(FDebugRenderSceneProxy::FText3d(PrintString, Point, Color));
+					DebugShapes.Texts.Add(FDebugRenderSceneProxy::FText3d(PrintString, Point, Color));
 				}
 			}
 		}
@@ -318,11 +324,11 @@ void AVisualLoggerRenderingActor::OnItemSelectionChanged(const FVisualLogDevice:
 		{
 			FDebugRenderSceneProxy::FMesh TestMesh;
 			GetPolygonMesh(ElementToDraw, TestMesh, CorridorOffset);
-			Meshes.Add(TestMesh);
+			DebugShapes.Meshes.Add(TestMesh);
 
 			for (int32 VIdx = 0; VIdx < ElementToDraw->Points.Num(); VIdx++)
 			{
-				Lines.Add(FDebugRenderSceneProxy::FDebugLine(
+				DebugShapes.Lines.Add(FDebugRenderSceneProxy::FDebugLine(
 					ElementToDraw->Points[VIdx] + CorridorOffset,
 					ElementToDraw->Points[(VIdx + 1) % ElementToDraw->Points.Num()] + CorridorOffset,
 					FColor::Cyan,
@@ -359,7 +365,7 @@ void AVisualLoggerRenderingActor::OnItemSelectionChanged(const FVisualLogDevice:
 				TestMesh.Indices.Add(CurrentFace.Y);
 				TestMesh.Indices.Add(CurrentFace.Z);
 			}
-			Meshes.Add(TestMesh);
+			DebugShapes.Meshes.Add(TestMesh);
 		}
 			break;
 		case EVisualLoggerShapeElement::Segment:
@@ -369,17 +375,17 @@ void AVisualLoggerRenderingActor::OnItemSelectionChanged(const FVisualLogDevice:
 			const FVector* Location = ElementToDraw->Points.GetData();
 			for (int32 Index = 0; Index + 1 < ElementToDraw->Points.Num(); Index += 2, Location += 2)
 			{
-				Lines.Add(FDebugRenderSceneProxy::FDebugLine(*Location, *(Location + 1), Color, Thickness));
+				DebugShapes.Lines.Add(FDebugRenderSceneProxy::FDebugLine(*Location, *(Location + 1), Color, Thickness));
 
 				if (bDrawLabel)
 				{
 					const FString PrintString = FString::Printf(TEXT("%s_%d"), *ElementToDraw->Description, Index);
-					Texts.Add(FDebugRenderSceneProxy::FText3d(PrintString, (*Location + (*(Location + 1) - *Location) / 2), Color));
+					DebugShapes.Texts.Add(FDebugRenderSceneProxy::FText3d(PrintString, (*Location + (*(Location + 1) - *Location) / 2), Color));
 				}
 			}
 			if (ElementToDraw->Description.IsEmpty() == false)
 			{
-				Texts.Add(FDebugRenderSceneProxy::FText3d(ElementToDraw->Description, ElementToDraw->Points[0] + (ElementToDraw->Points[1] - ElementToDraw->Points[0]) / 2, Color));
+				DebugShapes.Texts.Add(FDebugRenderSceneProxy::FText3d(ElementToDraw->Description, ElementToDraw->Points[0] + (ElementToDraw->Points[1] - ElementToDraw->Points[0]) / 2, Color));
 			}
 		}
 			break;
@@ -390,7 +396,7 @@ void AVisualLoggerRenderingActor::OnItemSelectionChanged(const FVisualLogDevice:
 			for (int32 Index = 1; Index < ElementToDraw->Points.Num(); ++Index)
 			{
 				const FVector CurrentLocation = ElementToDraw->Points[Index];
-				Lines.Add(FDebugRenderSceneProxy::FDebugLine(Location, CurrentLocation, Color, Thickness));
+				DebugShapes.Lines.Add(FDebugRenderSceneProxy::FDebugLine(Location, CurrentLocation, Color, Thickness));
 				Location = CurrentLocation;
 			}
 		}
@@ -403,17 +409,17 @@ void AVisualLoggerRenderingActor::OnItemSelectionChanged(const FVisualLogDevice:
 			for (int32 Index = 0; Index + 1 < ElementToDraw->Points.Num(); Index += 2, BoxExtent += 2)
 			{
 				const FBox Box = FBox(*BoxExtent, *(BoxExtent + 1));
-				Boxes.Add(FDebugRenderSceneProxy::FDebugBox(Box, Color, FTransform(ElementToDraw->TransformationMatrix)));
+				DebugShapes.Boxes.Add(FDebugRenderSceneProxy::FDebugBox(Box, Color, FTransform(ElementToDraw->TransformationMatrix)));
 
 				if (bDrawLabel)
 				{
 					const FString PrintString = FString::Printf(TEXT("%s_%d"), *ElementToDraw->Description, Index);
-					Texts.Add(FDebugRenderSceneProxy::FText3d(PrintString, Box.GetCenter(), Color));
+					DebugShapes.Texts.Add(FDebugRenderSceneProxy::FText3d(PrintString, Box.GetCenter(), Color));
 				}
 			}
 			if (ElementToDraw->Description.IsEmpty() == false)
 			{
-				Texts.Add(FDebugRenderSceneProxy::FText3d(ElementToDraw->Description, ElementToDraw->Points[0] + (ElementToDraw->Points[1] - ElementToDraw->Points[0]) / 2, Color));
+				DebugShapes.Texts.Add(FDebugRenderSceneProxy::FText3d(ElementToDraw->Description, ElementToDraw->Points[0] + (ElementToDraw->Points[1] - ElementToDraw->Points[0]) / 2, Color));
 			}
 		}
 			break;
@@ -430,11 +436,11 @@ void AVisualLoggerRenderingActor::OnItemSelectionChanged(const FVisualLogDevice:
 
 				FVector YAxis, ZAxis;
 				Direction.FindBestAxisVectors(YAxis, ZAxis);
-				Cones.Add(FDebugRenderSceneProxy::FCone(FScaleMatrix(FVector(Length)) * FMatrix(Direction, YAxis, ZAxis, Orgin), Angles.Y, Angles.Z, Color));
+				DebugShapes.Cones.Add(FDebugRenderSceneProxy::FCone(FScaleMatrix(FVector(Length)) * FMatrix(Direction, YAxis, ZAxis, Orgin), Angles.Y, Angles.Z, Color));
 
 				if (bDrawLabel)
 				{
-					Texts.Add(FDebugRenderSceneProxy::FText3d(ElementToDraw->Description, Orgin, Color));
+					DebugShapes.Texts.Add(FDebugRenderSceneProxy::FText3d(ElementToDraw->Description, Orgin, Color));
 				}
 			}
 		}
@@ -448,10 +454,10 @@ void AVisualLoggerRenderingActor::OnItemSelectionChanged(const FVisualLogDevice:
 				const FVector Start = ElementToDraw->Points[Index];
 				const FVector End = ElementToDraw->Points[Index + 1];
 				const FVector OtherData = ElementToDraw->Points[Index + 2];
-				Cylinders.Add(FDebugRenderSceneProxy::FWireCylinder(Start, OtherData.X, (End - Start).Size()*0.5, Color));
+				DebugShapes.Cylinders.Add(FDebugRenderSceneProxy::FWireCylinder(Start, OtherData.X, (End - Start).Size()*0.5, Color));
 				if (bDrawLabel)
 				{
-					Texts.Add(FDebugRenderSceneProxy::FText3d(ElementToDraw->Description, Start, Color));
+					DebugShapes.Texts.Add(FDebugRenderSceneProxy::FText3d(ElementToDraw->Description, Start, Color));
 				}
 			}
 		}
@@ -474,14 +480,14 @@ void AVisualLoggerRenderingActor::OnItemSelectionChanged(const FVisualLogDevice:
 				const FVector YAxis = Axes.GetScaledAxis(EAxis::Y);
 				const FVector ZAxis = Axes.GetScaledAxis(EAxis::Z);
 
-				Capsles.Add(FDebugRenderSceneProxy::FCapsule(Center, Radius, XAxis, YAxis, ZAxis, HalfHeight, Color));
+				DebugShapes.Capsles.Add(FDebugRenderSceneProxy::FCapsule(Center, Radius, XAxis, YAxis, ZAxis, HalfHeight, Color));
 				if (bDrawLabel)
 				{
-					Texts.Add(FDebugRenderSceneProxy::FText3d(ElementToDraw->Description, Center, Color));
+					DebugShapes.Texts.Add(FDebugRenderSceneProxy::FText3d(ElementToDraw->Description, Center, Color));
 				}
 			}
 		}
-		break;
+			break;
 		case EVisualLoggerShapeElement::NavAreaMesh:
 		{
 			if (ElementToDraw->Points.Num() == 0)
@@ -524,7 +530,7 @@ void AVisualLoggerRenderingActor::OnItemSelectionChanged(const FVisualLogDevice:
 				TestMesh.Indices.Add(CurrentIndex + 2);
 				CurrentIndex += 3;
 			}
-			Meshes.Add(TestMesh);
+			DebugShapes.Meshes.Add(TestMesh);
 
 			{
 				FDebugRenderSceneProxy::FMesh PolygonMesh;
@@ -532,22 +538,71 @@ void AVisualLoggerRenderingActor::OnItemSelectionChanged(const FVisualLogDevice:
 				PolygonToDraw.SetColor(ElementToDraw->GetFColor());
 				PolygonToDraw.Points.Reserve(AreaMeshPoints.Num());
 				PolygonToDraw.Points = AreaMeshPoints;
-				GetPolygonMesh(&PolygonToDraw, PolygonMesh, FVector(0,0,HeaderData.MaxZ));
-				Meshes.Add(PolygonMesh);
+				GetPolygonMesh(&PolygonToDraw, PolygonMesh, FVector(0, 0, HeaderData.MaxZ));
+				DebugShapes.Meshes.Add(PolygonMesh);
 			}
 
 			for (int32 VIdx = 0; VIdx < AreaMeshPoints.Num(); VIdx++)
 			{
-				Lines.Add(FDebugRenderSceneProxy::FDebugLine(
+				DebugShapes.Lines.Add(FDebugRenderSceneProxy::FDebugLine(
 					AreaMeshPoints[VIdx] + FVector(0, 0, HeaderData.MaxZ),
 					AreaMeshPoints[(VIdx + 1) % AreaMeshPoints.Num()] + FVector(0, 0, HeaderData.MaxZ),
 					ElementToDraw->GetFColor(),
 					2)
-				);
+					);
 			}
 
 		}
-		break;
+			break;
+		}
+	}
+}
+
+void AVisualLoggerRenderingActor::OnItemSelectionChanged(const FVisualLogDevice::FVisualLogEntryItem& EntryItem)
+{
+	PrimaryDebugShapes.Reset();
+	GetDebugShapes(EntryItem, PrimaryDebugShapes);
+
+	SecondaryDebugShapes.Reset();
+	if (SelectedTimeLines.Num() > 1)
+	{
+		for (TSharedPtr<class STimeline> CurrentTimeline : SelectedTimeLines)
+		{
+			if (CurrentTimeline.IsValid() == false || CurrentTimeline->GetVisibility().IsVisible() == false)
+			{
+				continue;
+			}
+
+			if (EntryItem.OwnerName == CurrentTimeline->GetName())
+			{
+				continue;
+			}
+
+			const TArray<FVisualLogDevice::FVisualLogEntryItem>& Entries = CurrentTimeline->GetEntries();
+
+			int32 BestItemIndex = INDEX_NONE;
+			float BestDistance = MAX_FLT;
+			for (int32 Index = 0; Index < Entries.Num(); Index++)
+			{
+				auto& CurrentEntryItem = Entries[Index];
+
+				if (CurrentTimeline->IsEntryHidden(CurrentEntryItem))
+				{
+					continue;
+				}
+				TArray<FVisualLoggerCategoryVerbosityPair> OutCategories;
+				const float CurrentDist = FMath::Abs(CurrentEntryItem.Entry.TimeStamp - EntryItem.Entry.TimeStamp);
+				if (CurrentDist < BestDistance)
+				{
+					BestDistance = CurrentDist;
+					BestItemIndex = Index;
+				}
+			}
+
+			if (Entries.IsValidIndex(BestItemIndex))
+			{
+				GetDebugShapes(Entries[BestItemIndex], SecondaryDebugShapes);
+			}
 		}
 	}
 
diff --git a/Engine/Source/Developer/LogVisualizer/Private/VisualLoggerRenderingActor.h b/Engine/Source/Developer/LogVisualizer/Private/VisualLoggerRenderingActor.h
--- a/Engine/Source/Developer/LogVisualizer/Private/VisualLoggerRenderingActor.h
+++ b/Engine/Source/Developer/LogVisualizer/Private/VisualLoggerRenderingActor.h
@@ -13,26 +13,49 @@ DECLARE_MULTICAST_DELEGATE_OneParam(FOnSelectionChanged, class AActor*);
 UCLASS(config = Engine, NotBlueprintable, Transient)
 class LOGVISUALIZER_API AVisualLoggerRenderingActor : public AActor
 {
+	struct FTimelineDebugShapes
+	{
+		TArray<FDebugRenderSceneProxy::FDebugLine> Lines;
+		TArray<FDebugRenderSceneProxy::FCone> Cones;
+		TArray<FDebugRenderSceneProxy::FDebugBox> Boxes;
+		TArray<FDebugRenderSceneProxy::FSphere> Points;
+		TArray<FDebugRenderSceneProxy::FMesh> Meshes;
+		TArray<FDebugRenderSceneProxy::FText3d> Texts;
+		TArray<FDebugRenderSceneProxy::FWireCylinder> Cylinders;
+		TArray<FDebugRenderSceneProxy::FCapsule> Capsles;
+		TArray<FVector> LogEntriesPath;
+
+		void Reset()
+		{
+			Lines.Reset();
+			Cones.Reset();
+			Boxes.Reset();
+			Points.Reset();
+			Meshes.Reset();
+			Texts.Reset();
+			Cylinders.Reset();
+			Capsles.Reset();
+			LogEntriesPath.Reset();
+		}
+	};
+
 public:
 	GENERATED_UCLASS_BODY()
 
 	void OnItemSelectionChanged(const FVisualLogDevice::FVisualLogEntryItem& EntryItem);
-	void ObjectSelectionChanged(TSharedPtr<class STimeline> TimeLine);
+	void ObjectSelectionChanged(TArray<TSharedPtr<class STimeline> >& TimeLines);
+
+	void GetDebugShapes(const FVisualLogDevice::FVisualLogEntryItem& EntryItem, FTimelineDebugShapes& DebugShapes);
 
 private:
 	void AddDebugRendering();
 
 public:
 	UPrimitiveComponent* RenderingComponent;
-	
-	TArray<FDebugRenderSceneProxy::FDebugLine> Lines;
-	TArray<FDebugRenderSceneProxy::FCone> Cones;
-	TArray<FDebugRenderSceneProxy::FDebugBox> Boxes;
-	TArray<FDebugRenderSceneProxy::FSphere> Points;
-	TArray<FDebugRenderSceneProxy::FMesh> Meshes;
-	TArray<FDebugRenderSceneProxy::FText3d> Texts;
-	TArray<FDebugRenderSceneProxy::FWireCylinder> Cylinders;
-	TArray<FDebugRenderSceneProxy::FCapsule> Capsles;
-	TArray<FVector> LogEntriesPath;
 
+	FTimelineDebugShapes PrimaryDebugShapes;
+	FTimelineDebugShapes SecondaryDebugShapes;
+	FTimelineDebugShapes TestDebugShapes;
+
+	TArray<TSharedPtr<class STimeline> > SelectedTimeLines;
 };
