Fixed status view categories in Visual Logger tool.

diff --git a/Engine/Source/Developer/LogVisualizer/Private/SVisualLoggerStatusView.cpp b/Engine/Source/Developer/LogVisualizer/Private/SVisualLoggerStatusView.cpp
--- a/Engine/Source/Developer/LogVisualizer/Private/SVisualLoggerStatusView.cpp
+++ b/Engine/Source/Developer/LogVisualizer/Private/SVisualLoggerStatusView.cpp
@@ -86,7 +86,10 @@ void GenerateChildren(TSharedPtr<FLogStatusItem> StatusItem, const FVisualLogSta
 void SVisualLoggerStatusView::ObjectSelectionChanged(TArray<TSharedPtr<class STimeline> >& TimeLines)
 {
 	SelectedTimeLines = TimeLines;
-	StatusItems.Reset();
+	if (SelectedTimeLines.Num() == 0)
+	{
+		StatusItems.Reset();
+	}
 	StatusItemsView->RequestTreeRefresh();
 }
 
