Bunch of fixes to have better control over Visual Logger tool with keyboard.

diff --git a/Engine/Source/Developer/LogVisualizer/Private/LogVisualizer.cpp b/Engine/Source/Developer/LogVisualizer/Private/LogVisualizer.cpp
--- a/Engine/Source/Developer/LogVisualizer/Private/LogVisualizer.cpp
+++ b/Engine/Source/Developer/LogVisualizer/Private/LogVisualizer.cpp
@@ -103,22 +103,31 @@ void FLogVisualizer::Goto(float Timestamp, FName LogOwner)
 	}
 }
 
-void FLogVisualizer::GotoNextItem()
+void FLogVisualizer::GotoNextItem(int32 Distance)
 {
 	if (CurrentTimeLine.IsValid())
 	{
-		CurrentTimeLine.Pin()->GotoNextItem();
+		CurrentTimeLine.Pin()->MoveCursorByDistance(Distance);
 	}
 }
 
-void FLogVisualizer::GotoPreviousItem()
+void FLogVisualizer::GotoPreviousItem(int32 Distance)
 {
 	if (CurrentTimeLine.IsValid())
 	{
-		CurrentTimeLine.Pin()->GotoPreviousItem();
+		CurrentTimeLine.Pin()->MoveCursorByDistance(-Distance);
 	}
 }
 
+void FLogVisualizer::MoveCamera()
+{
+	if (CurrentTimeLine.IsValid())
+	{
+		CurrentTimeLine.Pin()->UpdateCameraPosition();
+	}
+}
+
+
 FLinearColor FLogVisualizer::GetColorForCategory(int32 Index) const
 {
 	if (Index >= 0 && Index < sizeof(FLogVisualizerColorPalette) / sizeof(FLogVisualizerColorPalette[0]))
diff --git a/Engine/Source/Developer/LogVisualizer/Private/LogVisualizer.h b/Engine/Source/Developer/LogVisualizer/Private/LogVisualizer.h
--- a/Engine/Source/Developer/LogVisualizer/Private/LogVisualizer.h
+++ b/Engine/Source/Developer/LogVisualizer/Private/LogVisualizer.h
@@ -9,6 +9,7 @@
 #include "SlateBasics.h"
 #include "SlateStyle.h"
 #include "EditorStyle.h"
+#include "SDockTab.h"
 
 /* Private includes
 *****************************************************************************/
@@ -22,6 +23,7 @@ DECLARE_DELEGATE_OneParam(FOnObjectSelectionChanged, TSharedPtr<class STimeline>
 DECLARE_DELEGATE_OneParam(FOnFiltersSearchChanged, const FText&);
 DECLARE_DELEGATE(FOnFiltersChanged);
 DECLARE_DELEGATE_ThreeParams(FOnLogLineSelectionChanged, TSharedPtr<FLogEntryItem> /*SelectedItem*/, int64 /*UserData*/, FName /*TagName*/);
+DECLARE_DELEGATE_RetVal_TwoParams(FReply, FOnKeyboardEvent, const FGeometry& /*MyGeometry*/, const FKeyEvent& /*InKeyEvent*/);
 
 struct FVisualLoggerEvents
 {
@@ -34,6 +36,7 @@ struct FVisualLoggerEvents
 	FOnFiltersChanged OnFiltersChanged;
 	FOnObjectSelectionChanged OnObjectSelectionChanged;
 	FOnLogLineSelectionChanged OnLogLineSelectionChanged;
+	FOnKeyboardEvent OnKeyboardEvent;
 };
 
 class FVisualLoggerTimeSliderController;
@@ -41,8 +44,9 @@ struct LOGVISUALIZER_API FLogVisualizer
 {
 	/** LogVisualizer interface*/
 	void Goto(float Timestamp, FName LogOwner = NAME_None);
-	void GotoNextItem();
-	void GotoPreviousItem();
+	void GotoNextItem( int32 Distance);
+	void GotoPreviousItem(int32 Distance);
+	void MoveCamera();
 
 	FLinearColor GetColorForCategory(int32 Index) const;
 	FLinearColor GetColorForCategory(const FString& InFilterName) const;
@@ -68,6 +72,27 @@ protected:
 	TWeakPtr<class SVisualLogger> CurrentVisualizer;
 };
 
+class SVisualLoggerTab : public SDockTab
+{
+public:
+	virtual bool SupportsKeyboardFocus() const override { return true; }
+	virtual FReply OnKeyDown(const FGeometry& MyGeometry, const FKeyEvent& InKeyEvent) override
+	{
+		return FLogVisualizer::Get().GetVisualLoggerEvents().OnKeyboardEvent.Execute(MyGeometry, InKeyEvent);
+	}
+};
+
+class SVisualLoggerBaseWidget : public SCompoundWidget
+{
+public:
+	virtual bool SupportsKeyboardFocus() const override { return true; }
+	virtual FReply OnKeyDown(const FGeometry& MyGeometry, const FKeyEvent& InKeyEvent) override
+	{
+		return FLogVisualizer::Get().GetVisualLoggerEvents().OnKeyboardEvent.Execute(MyGeometry, InKeyEvent);
+	}
+};
+
+
 #include "LogVisualizerStyle.h"
 #include "SVisualLogger.h"
 #include "SVisualLoggerToolbar.h"
diff --git a/Engine/Source/Developer/LogVisualizer/Private/LogVisualizerModule.cpp b/Engine/Source/Developer/LogVisualizer/Private/LogVisualizerModule.cpp
--- a/Engine/Source/Developer/LogVisualizer/Private/LogVisualizerModule.cpp
+++ b/Engine/Source/Developer/LogVisualizer/Private/LogVisualizerModule.cpp
@@ -3,7 +3,6 @@
 #include "LogVisualizer.h"
 #include "Runtime/Core/Public/Features/IModularFeatures.h"
 #include "LogVisualizerStyle.h"
-#include "SDockTab.h"
 #include "VisualLoggerRenderingActor.h"
 #include "LogVisualizerSettings.h"
 #if WITH_EDITOR
@@ -80,7 +79,7 @@ void FLogVisualizerModule::ShutdownModule()
 
 TSharedRef<SDockTab> FLogVisualizerModule::SpawnLogVisualizerTab(const FSpawnTabArgs& SpawnTabArgs)
 {
-	const TSharedRef<SDockTab> MajorTab = SNew(SDockTab)
+	const TSharedRef<SDockTab> MajorTab = SNew(SVisualLoggerTab)
 		.TabRole(ETabRole::NomadTab);
 
 	TSharedPtr<SWidget> TabContent;
@@ -99,12 +98,12 @@ void FLogVisualizerModule::Goto(float Timestamp, FName LogOwner)
 
 void FLogVisualizerModule::GotoNextItem()
 {
-	FLogVisualizer::Get().GotoNextItem();
+	FLogVisualizer::Get().GotoNextItem(1);
 }
 
 void FLogVisualizerModule::GotoPreviousItem()
 {
-	FLogVisualizer::Get().GotoPreviousItem();
+	FLogVisualizer::Get().GotoPreviousItem(1);
 }
 
 
diff --git a/Engine/Source/Developer/LogVisualizer/Private/STimeline.cpp b/Engine/Source/Developer/LogVisualizer/Private/STimeline.cpp
--- a/Engine/Source/Developer/LogVisualizer/Private/STimeline.cpp
+++ b/Engine/Source/Developer/LogVisualizer/Private/STimeline.cpp
@@ -287,3 +287,22 @@ void STimeline::GotoPreviousItem()
 { 
 	TimelineBar->GotoPreviousItem(); 
 }
+
+void STimeline::MoveCursorByDistance(int32 Distance)
+{
+	if (Distance > 0)
+	{
+		TimelineBar->GotoNextItem(FMath::Abs(Distance));
+	}
+	else
+	{
+		TimelineBar->GotoPreviousItem(FMath::Abs(Distance));
+	}
+	
+}
+
+
+void STimeline::UpdateCameraPosition()
+{
+	TimelineBar->UpdateCameraPosition();
+}
diff --git a/Engine/Source/Developer/LogVisualizer/Private/STimeline.h b/Engine/Source/Developer/LogVisualizer/Private/STimeline.h
--- a/Engine/Source/Developer/LogVisualizer/Private/STimeline.h
+++ b/Engine/Source/Developer/LogVisualizer/Private/STimeline.h
@@ -43,6 +43,8 @@ public:
 	void Goto(float ScrubPosition);
 	void GotoNextItem();
 	void GotoPreviousItem();
+	void MoveCursorByDistance(int32 Distance);
+	void UpdateCameraPosition();
 
 protected:
 	TArray<FVisualLogDevice::FVisualLogEntryItem> Entries;
diff --git a/Engine/Source/Developer/LogVisualizer/Private/STimelineBar.cpp b/Engine/Source/Developer/LogVisualizer/Private/STimelineBar.cpp
--- a/Engine/Source/Developer/LogVisualizer/Private/STimelineBar.cpp
+++ b/Engine/Source/Developer/LogVisualizer/Private/STimelineBar.cpp
@@ -39,70 +39,78 @@ FReply STimelineBar::OnMouseMove(const FGeometry& MyGeometry, const FPointerEven
 	return TimeSliderController->OnMouseMove(SharedThis(this), MyGeometry, MouseEvent);
 }
 
-FReply STimelineBar::OnMouseButtonDoubleClick(const FGeometry& MyGeometry, const FPointerEvent& MouseEvent)
+void STimelineBar::UpdateCameraPosition()
 {
 	auto &Entries = TimelineOwner.Pin()->GetEntries();
-	SnapScrubPosition(TimeSliderController->GetTimeSliderArgs().ScrubPosition.Get());
 	UWorld* World = FLogVisualizer::Get().GetWorld();
-	if (World && MouseEvent.GetEffectingButton() == EKeys::LeftMouseButton && Entries.IsValidIndex(CurrentItemIndex))
-	{
-		FVector CurrentLocation = Entries[CurrentItemIndex].Entry.Location;
 
-		FVector Extent(150);
-		bool bFoundActor = false;
-		FName OwnerName = Entries[CurrentItemIndex].OwnerName;
-		for (FActorIterator It(World); It; ++It)
+	FVector CurrentLocation = Entries[CurrentItemIndex].Entry.Location;
+
+	FVector Extent(150);
+	bool bFoundActor = false;
+	FName OwnerName = Entries[CurrentItemIndex].OwnerName;
+	for (FActorIterator It(World); It; ++It)
+	{
+		AActor* Actor = *It;
+		if (Actor->GetFName() == OwnerName)
 		{
-			AActor* Actor = *It;
-			if (Actor->GetFName() == OwnerName)
-			{
-				FVector Orgin;
-				Actor->GetActorBounds(false, Orgin, Extent);
-				bFoundActor = true;
-				break;
-			}
+			FVector Orgin;
+			Actor->GetActorBounds(false, Orgin, Extent);
+			bFoundActor = true;
+			break;
 		}
+	}
+
 
-		
-		const float DefaultCameraDistance = ULogVisualizerSettings::StaticClass()->GetDefaultObject<ULogVisualizerSettings>()->DefaultCameraDistance;
-		Extent = Extent.Size() < DefaultCameraDistance ? FVector(1) * DefaultCameraDistance : Extent;
+	const float DefaultCameraDistance = ULogVisualizerSettings::StaticClass()->GetDefaultObject<ULogVisualizerSettings>()->DefaultCameraDistance;
+	Extent = Extent.Size() < DefaultCameraDistance ? FVector(1) * DefaultCameraDistance : Extent;
 
 #if WITH_EDITOR
-		UEditorEngine *EEngine = Cast<UEditorEngine>(GEngine);
-		if (GIsEditor && EEngine != NULL)
+	UEditorEngine *EEngine = Cast<UEditorEngine>(GEngine);
+	if (GIsEditor && EEngine != NULL)
+	{
+		for (auto ViewportClient : EEngine->AllViewportClients)
 		{
-			for (auto ViewportClient : EEngine->AllViewportClients)
-			{
-				ViewportClient->FocusViewportOnBox(FBox::BuildAABB(CurrentLocation, Extent));
-			}
+			ViewportClient->FocusViewportOnBox(FBox::BuildAABB(CurrentLocation, Extent));
 		}
-		else if (AVisualLoggerCameraController::IsEnabled(World) && AVisualLoggerCameraController::Instance.IsValid() && AVisualLoggerCameraController::Instance->GetSpectatorPawn())
+	}
+	else if (AVisualLoggerCameraController::IsEnabled(World) && AVisualLoggerCameraController::Instance.IsValid() && AVisualLoggerCameraController::Instance->GetSpectatorPawn())
+	{
+		ULocalPlayer* LocalPlayer = Cast<ULocalPlayer>(AVisualLoggerCameraController::Instance->Player);
+		if (LocalPlayer && LocalPlayer->ViewportClient && LocalPlayer->ViewportClient->Viewport)
 		{
-			ULocalPlayer* LocalPlayer = Cast<ULocalPlayer>(AVisualLoggerCameraController::Instance->Player);
-			if (LocalPlayer && LocalPlayer->ViewportClient && LocalPlayer->ViewportClient->Viewport)
-			{
-				FViewport* Viewport = LocalPlayer->ViewportClient->Viewport;
+			FViewport* Viewport = LocalPlayer->ViewportClient->Viewport;
 
-				FBox BoundingBox = FBox::BuildAABB(CurrentLocation, Extent);
-				const FVector Position = BoundingBox.GetCenter();
-				float Radius = BoundingBox.GetExtent().Size();
+			FBox BoundingBox = FBox::BuildAABB(CurrentLocation, Extent);
+			const FVector Position = BoundingBox.GetCenter();
+			float Radius = BoundingBox.GetExtent().Size();
 
-				FViewportCameraTransform ViewTransform;
-				ViewTransform.TransitionToLocation(Position, SharedThis(this), true);
+			FViewportCameraTransform ViewTransform;
+			ViewTransform.TransitionToLocation(Position, SharedThis(this), true);
 
-				float NewOrthoZoom;
-				const float AspectRatio = 1.777777f;
-				uint32 MinAxisSize = (AspectRatio > 1.0f) ? Viewport->GetSizeXY().Y : Viewport->GetSizeXY().X;
-				float Zoom = Radius / (MinAxisSize / 2.0f);
+			float NewOrthoZoom;
+			const float AspectRatio = 1.777777f;
+			uint32 MinAxisSize = (AspectRatio > 1.0f) ? Viewport->GetSizeXY().Y : Viewport->GetSizeXY().X;
+			float Zoom = Radius / (MinAxisSize / 2.0f);
 
-				NewOrthoZoom = Zoom * (Viewport->GetSizeXY().X*15.0f);
-				NewOrthoZoom = FMath::Clamp<float>(NewOrthoZoom, 250, MAX_FLT);
-				ViewTransform.SetOrthoZoom(NewOrthoZoom);
+			NewOrthoZoom = Zoom * (Viewport->GetSizeXY().X*15.0f);
+			NewOrthoZoom = FMath::Clamp<float>(NewOrthoZoom, 250, MAX_FLT);
+			ViewTransform.SetOrthoZoom(NewOrthoZoom);
 
-				AVisualLoggerCameraController::Instance->GetSpectatorPawn()->TeleportTo(ViewTransform.GetLocation(), ViewTransform.GetRotation(), false, true);
-			}
+			AVisualLoggerCameraController::Instance->GetSpectatorPawn()->TeleportTo(ViewTransform.GetLocation(), ViewTransform.GetRotation(), false, true);
 		}
+	}
 #endif
+}
+
+FReply STimelineBar::OnMouseButtonDoubleClick(const FGeometry& MyGeometry, const FPointerEvent& MouseEvent)
+{
+	auto &Entries = TimelineOwner.Pin()->GetEntries();
+	SnapScrubPosition(TimeSliderController->GetTimeSliderArgs().ScrubPosition.Get());
+	UWorld* World = FLogVisualizer::Get().GetWorld();
+	if (World && MouseEvent.GetEffectingButton() == EKeys::LeftMouseButton && Entries.IsValidIndex(CurrentItemIndex))
+	{
+		UpdateCameraPosition();
 		return FReply::Handled();
 	}
 
diff --git a/Engine/Source/Developer/LogVisualizer/Private/STimelineBar.h b/Engine/Source/Developer/LogVisualizer/Private/STimelineBar.h
--- a/Engine/Source/Developer/LogVisualizer/Private/STimelineBar.h
+++ b/Engine/Source/Developer/LogVisualizer/Private/STimelineBar.h
@@ -33,6 +33,7 @@ public:
 	void OnDeselect();
 	void GotoNextItem(int32 MoveDistance = 1);
 	void GotoPreviousItem(int32 MoveDistance = 1);
+	void UpdateCameraPosition();
 
 protected:
 	TSharedPtr<class FVisualLoggerTimeSliderController> TimeSliderController;
diff --git a/Engine/Source/Developer/LogVisualizer/Private/STimelinesContainer.cpp b/Engine/Source/Developer/LogVisualizer/Private/STimelinesContainer.cpp
--- a/Engine/Source/Developer/LogVisualizer/Private/STimelinesContainer.cpp
+++ b/Engine/Source/Developer/LogVisualizer/Private/STimelinesContainer.cpp
@@ -213,7 +213,6 @@ FReply STimelinesContainer::OnKeyDown(const FGeometry& MyGeometry, const FKeyEve
 		SetSelectionState(NotSelectedOne.Pin(), true, true);
 		return FReply::Handled();
 	}
-#if 0 //disable movement between timelines for now
 	else if (InKeyEvent.GetKey() == EKeys::Up || InKeyEvent.GetKey() == EKeys::Down)
 	{
 		TSharedPtr<class STimeline> PreviousTimeline;
@@ -252,7 +251,7 @@ FReply STimelinesContainer::OnKeyDown(const FGeometry& MyGeometry, const FKeyEve
 		}
 		return FReply::Handled();
 	}
-#endif
+
 	return FReply::Unhandled();
 }
 
diff --git a/Engine/Source/Developer/LogVisualizer/Private/SVisualLogger.cpp b/Engine/Source/Developer/LogVisualizer/Private/SVisualLogger.cpp
--- a/Engine/Source/Developer/LogVisualizer/Private/SVisualLogger.cpp
+++ b/Engine/Source/Developer/LogVisualizer/Private/SVisualLogger.cpp
@@ -130,7 +130,7 @@ void SVisualLogger::Construct(const FArguments& InArgs, const TSharedRef<SDockTa
 	FLogVisualizer::Get().GetVisualLoggerEvents().OnFiltersChanged = FOnFiltersChanged::CreateRaw(this, &SVisualLogger::OnFiltersChanged);
 	FLogVisualizer::Get().GetVisualLoggerEvents().OnObjectSelectionChanged = FOnObjectSelectionChanged::CreateRaw(this, &SVisualLogger::OnObjectSelectionChanged);
 	FLogVisualizer::Get().GetVisualLoggerEvents().OnLogLineSelectionChanged = FOnLogLineSelectionChanged::CreateRaw(this, &SVisualLogger::OnLogLineSelectionChanged);
-
+	FLogVisualizer::Get().GetVisualLoggerEvents().OnKeyboardEvent = FOnKeyboardEvent::CreateRaw(this, &SVisualLogger::OnKeyboaedRedirection);
 	//////////////////////////////////////////////////////////////////////////
 	// Command Action Lists
 	const FVisualLoggerCommands& Commands = FVisualLoggerCommands::Get();
@@ -307,12 +307,6 @@ void SVisualLogger::Construct(const FArguments& InArgs, const TSharedRef<SDockTa
 	//UGameplayDebuggingComponent::OnDebuggingTargetChangedDelegate.AddSP(this, &SVisualLogger::SelectionChanged);
 }
 
-FReply SVisualLogger::OnKeyDown(const FGeometry& MyGeometry, const FKeyEvent& InKeyEvent)
-{
-	FUICommandList& ActionList = *CommandList;
-	return ActionList.ProcessCommandBindings(InKeyEvent) ? FReply::Handled() : FReply::Unhandled();
-}
-
 void SVisualLogger::HandleMajorTabPersistVisualState()
 {
 	// save any settings here
@@ -389,7 +383,7 @@ TSharedRef<SDockTab> SVisualLogger::HandleTabManagerSpawnTab(const FSpawnTabArgs
 	bool AutoSizeTab = false;
 
 	if (TabIdentifier == ToolbarTabId)
-	{
+	{ 
 		TabWidget = SNew(SVisualLoggerToolbar, CommandList);
 		AutoSizeTab = true;
 	}
@@ -413,10 +407,10 @@ TSharedRef<SDockTab> SVisualLogger::HandleTabManagerSpawnTab(const FSpawnTabArgs
 		TabWidget = SAssignNew(StatusView, SVisualLoggerStatusView, CommandList);
 		AutoSizeTab = false;
 	}
-	
-	return SNew(SDockTab)
+
+	return SNew(SVisualLoggerTab)
 		.ShouldAutosize(AutoSizeTab)
-		.TabRole(ETabRole::PanelTab)
+		.TabRole(ETabRole::DocumentTab)
 		[
 			TabWidget.ToSharedRef()
 		];
@@ -866,12 +860,38 @@ void SVisualLogger::OnLogLineSelectionChanged(TSharedPtr<struct FLogEntryItem> S
 
 void SVisualLogger::OnMoveCursorLeftCommand()
 {
-	FLogVisualizer::Get().GotoNextItem();
+	FLogVisualizer::Get().GotoNextItem(1);
 }
 
 void SVisualLogger::OnMoveCursorRightCommand()
 {
-	FLogVisualizer::Get().GotoPreviousItem();
+	FLogVisualizer::Get().GotoPreviousItem(1);
+}
+
+FReply SVisualLogger::OnKeyboaedRedirection(const FGeometry& MyGeometry, const FKeyEvent& InKeyEvent)
+{
+	FReply ReturnValue = FReply::Unhandled();
+
+	if (InKeyEvent.GetKey() == EKeys::Left)
+	{
+		FLogVisualizer::Get().GotoPreviousItem(InKeyEvent.IsLeftControlDown() ? InKeyEvent.IsLeftShiftDown() ? 20 : 10 : 1);
+		ReturnValue = FReply::Handled();
+	}
+	else if (InKeyEvent.GetKey() == EKeys::Right)
+	{
+		FLogVisualizer::Get().GotoNextItem(InKeyEvent.IsLeftControlDown() ? InKeyEvent.IsLeftShiftDown() ? 20 : 10 : 1);
+		ReturnValue = FReply::Handled();
+	}
+	else if (InKeyEvent.GetKey() == EKeys::Enter)
+	{
+		FLogVisualizer::Get().MoveCamera();
+		ReturnValue = FReply::Handled();
+	}
+	else if (InKeyEvent.GetKey() == EKeys::Up || InKeyEvent.GetKey() == EKeys::Down)
+	{
+		ReturnValue = FReply::Handled();
+	}
+	return ReturnValue;
 }
 
 void SVisualLogger::GetTimelines(TArray<TSharedPtr<class STimeline> >& TimeLines, bool bOnlySelectedOnes)
diff --git a/Engine/Source/Developer/LogVisualizer/Private/SVisualLogger.h b/Engine/Source/Developer/LogVisualizer/Private/SVisualLogger.h
--- a/Engine/Source/Developer/LogVisualizer/Private/SVisualLogger.h
+++ b/Engine/Source/Developer/LogVisualizer/Private/SVisualLogger.h
@@ -60,7 +60,7 @@ public:
 
 	/** Callback for for when the owner tab's visual state is being persisted. */
 	void HandleMajorTabPersistVisualState();
-	virtual FReply OnKeyDown(const FGeometry& MyGeometry, const FKeyEvent& InKeyEvent) override;
+	virtual bool SupportsKeyboardFocus() const override { return true; }
 
 	void OnNewLogEntry(const FVisualLogDevice::FVisualLogEntryItem& Entry);
 	void CollectNewCategories(const FVisualLogDevice::FVisualLogEntryItem& Entry);
@@ -69,6 +69,8 @@ public:
 	void OnObjectSelectionChanged(TSharedPtr<class STimeline> TimeLine);
 	void OnFiltersSearchChanged(const FText& Filter);
 	void OnLogLineSelectionChanged(TSharedPtr<struct FLogEntryItem> SelectedItem, int64 UserData, FName TagName);
+	FReply OnKeyboaedRedirection(const FGeometry& MyGeometry, const FKeyEvent& InKeyEvent);
+
 
 	bool HandleStartRecordingCommandCanExecute() const;
 	void HandleStartRecordingCommandExecute();
diff --git a/Engine/Source/Developer/LogVisualizer/Private/SVisualLoggerFilters.h b/Engine/Source/Developer/LogVisualizer/Private/SVisualLoggerFilters.h
--- a/Engine/Source/Developer/LogVisualizer/Private/SVisualLoggerFilters.h
+++ b/Engine/Source/Developer/LogVisualizer/Private/SVisualLoggerFilters.h
@@ -6,7 +6,7 @@
 * A list of filters currently applied to an asset view.
 */
 class SFilterWidget;
-class SVisualLoggerFilters : public SCompoundWidget
+class SVisualLoggerFilters : public SVisualLoggerBaseWidget
 {
 public:
 	SLATE_BEGIN_ARGS(SVisualLoggerFilters) { }
diff --git a/Engine/Source/Developer/LogVisualizer/Private/SVisualLoggerLogsList.cpp b/Engine/Source/Developer/LogVisualizer/Private/SVisualLoggerLogsList.cpp
--- a/Engine/Source/Developer/LogVisualizer/Private/SVisualLoggerLogsList.cpp
+++ b/Engine/Source/Developer/LogVisualizer/Private/SVisualLoggerLogsList.cpp
@@ -64,7 +64,7 @@ FReply SVisualLoggerLogsList::OnKeyDown(const FGeometry& MyGeometry, const FKeyE
 		return FReply::Handled();
 	}
 
-	return FReply::Unhandled();
+	return SVisualLoggerBaseWidget::OnKeyDown(MyGeometry, InKeyEvent);
 }
 
 void SVisualLoggerLogsList::OnFiltersChanged()
diff --git a/Engine/Source/Developer/LogVisualizer/Private/SVisualLoggerLogsList.h b/Engine/Source/Developer/LogVisualizer/Private/SVisualLoggerLogsList.h
--- a/Engine/Source/Developer/LogVisualizer/Private/SVisualLoggerLogsList.h
+++ b/Engine/Source/Developer/LogVisualizer/Private/SVisualLoggerLogsList.h
@@ -2,7 +2,7 @@
 
 #pragma once
 
-class SVisualLoggerLogsList : public SCompoundWidget
+class SVisualLoggerLogsList : public SVisualLoggerBaseWidget
 {
 public:
 	SLATE_BEGIN_ARGS(SVisualLoggerLogsList){}
@@ -12,7 +12,6 @@ public:
 
 	TSharedRef<ITableRow> LogEntryLinesGenerateRow(TSharedPtr<struct FLogEntryItem> Item, const TSharedRef<STableViewBase>& OwnerTable);
 	void LogEntryLineSelectionChanged(TSharedPtr<FLogEntryItem> SelectedItem, ESelectInfo::Type SelectInfo);
-	virtual bool SupportsKeyboardFocus() const override { return true; }
 	virtual FReply OnKeyDown(const FGeometry& MyGeometry, const FKeyEvent& InKeyEvent) override;
 	void ObjectSelectionChanged(TArray<TSharedPtr<class STimeline> >& TimeLines);
 
diff --git a/Engine/Source/Developer/LogVisualizer/Private/SVisualLoggerStatusView.h b/Engine/Source/Developer/LogVisualizer/Private/SVisualLoggerStatusView.h
--- a/Engine/Source/Developer/LogVisualizer/Private/SVisualLoggerStatusView.h
+++ b/Engine/Source/Developer/LogVisualizer/Private/SVisualLoggerStatusView.h
@@ -2,7 +2,7 @@
 
 #pragma once
 
-class SVisualLoggerStatusView : public SCompoundWidget
+class SVisualLoggerStatusView : public SVisualLoggerBaseWidget
 {
 public:
 	SLATE_BEGIN_ARGS(SVisualLoggerStatusView){}
diff --git a/Engine/Source/Developer/LogVisualizer/Private/SVisualLoggerToolbar.h b/Engine/Source/Developer/LogVisualizer/Private/SVisualLoggerToolbar.h
--- a/Engine/Source/Developer/LogVisualizer/Private/SVisualLoggerToolbar.h
+++ b/Engine/Source/Developer/LogVisualizer/Private/SVisualLoggerToolbar.h
@@ -6,7 +6,7 @@
 /**
 * Implements the messaging debugger toolbar widget.
 */
-class SVisualLoggerToolbar : public SCompoundWidget
+class SVisualLoggerToolbar : public SVisualLoggerBaseWidget
 {
 public:
 
diff --git a/Engine/Source/Developer/LogVisualizer/Private/SVisualLoggerView.h b/Engine/Source/Developer/LogVisualizer/Private/SVisualLoggerView.h
--- a/Engine/Source/Developer/LogVisualizer/Private/SVisualLoggerView.h
+++ b/Engine/Source/Developer/LogVisualizer/Private/SVisualLoggerView.h
@@ -2,7 +2,7 @@
 
 #pragma once
 
-class SVisualLoggerView : public SCompoundWidget
+class SVisualLoggerView : public SVisualLoggerBaseWidget
 {
 public:
 	SLATE_BEGIN_ARGS(SVisualLoggerView) 
