Remove invalid sequencer todo's

diff --git a/Engine/Source/Developer/LogVisualizer/Private/SVisualLoggerView.cpp b/Engine/Source/Developer/LogVisualizer/Private/SVisualLoggerView.cpp
--- a/Engine/Source/Developer/LogVisualizer/Private/SVisualLoggerView.cpp
+++ b/Engine/Source/Developer/LogVisualizer/Private/SVisualLoggerView.cpp
@@ -174,7 +174,6 @@ void SVisualLoggerView::Construct(const FArguments& InArgs, const TSharedRef<FUI
 							[
 #if 0 //top time slider disabled to test idea with filter's search box
 								SNew(SBorder)
-								// @todo Sequencer Do not change the paddings or the sliders scrub widgets wont line up
 								.Padding(FMargin(0.0f, 2.0f, 0.0f, 0.0f))
 								.BorderImage(FLogVisualizerStyle::Get().GetBrush("ToolPanel.GroupBorder"))
 								.BorderBackgroundColor(FLinearColor(.50f, .50f, .50f, 1.0f))
@@ -245,7 +244,6 @@ void SVisualLoggerView::Construct(const FArguments& InArgs, const TSharedRef<FUI
 						.FillWidth(1.0f)
 						[
 							SNew(SBorder)
-							// @todo Sequencer Do not change the paddings or the sliders scrub widgets wont line up
 							.Padding(FMargin(0.0f, 0.0f, 0.0f, 2.0f))
 							.BorderImage(FLogVisualizerStyle::Get().GetBrush("ToolPanel.GroupBorder"))
 							.BorderBackgroundColor(FLinearColor(.50f, .50f, .50f, 1.0f))
