Optimizations to Visual Logger edytor tool:
Much faster filters due to Async framework, it's paraller now,
Optimizations to items rendered on sequencer,
Optimization to graphs rendered on canvas,
Changed as much as possible to event driven flot, to update data only when needed.

diff --git a/Engine/Source/Developer/LogVisualizer/Private/LogVisualizer.cpp b/Engine/Source/Developer/LogVisualizer/Private/LogVisualizer.cpp
--- a/Engine/Source/Developer/LogVisualizer/Private/LogVisualizer.cpp
+++ b/Engine/Source/Developer/LogVisualizer/Private/LogVisualizer.cpp
@@ -4,6 +4,7 @@
 #include "TimeSliderController.h"
 #include "VisualLoggerRenderingActor.h"
 #include "STimelinesContainer.h"
+#include "VisualLoggerCameraController.h"
 #if WITH_EDITOR
 #	include "Editor/UnrealEd/Public/EditorComponents.h"
 #	include "Editor/UnrealEd/Public/EditorReimportHandler.h"
@@ -68,7 +69,6 @@ void FLogVisualizer::Initialize()
 {
 	StaticInstance = MakeShareable(new FLogVisualizer);
 	Get().TimeSliderController = MakeShareable(new FVisualLoggerTimeSliderController(FVisualLoggerTimeSliderArgs()));
-
 }
 
 void FLogVisualizer::Shutdown()
@@ -76,58 +76,16 @@ void FLogVisualizer::Shutdown()
 	StaticInstance.Reset();
 }
 
-FLogVisualizer& FLogVisualizer::Get()
+inline void FLogVisualizer::Reset()
 {
-	return *StaticInstance;
+	TimeSliderController->SetTimesliderArgs(FVisualLoggerTimeSliderArgs());
 }
 
-void FLogVisualizer::Goto(float Timestamp, FName LogOwner)
-{
-	if (CurrentVisualizer.IsValid())
-	{
-		TArray<TSharedPtr<class STimeline> > TimeLines;
-		CurrentVisualizer.Pin()->GetTimelines(TimeLines, false);
-		for (TSharedPtr<class STimeline> Timeline : TimeLines)
-		{
-			if (Timeline->GetName() == LogOwner)
-			{
-				Timeline->GetOwner()->SetSelectionState(Timeline, true, true);
-				break;
-			}
-		}
-	}
-
-	if (CurrentTimeLine.IsValid())
-	{
-		CurrentTimeLine.Pin()->Goto(Timestamp);
-	}
-}
-
-void FLogVisualizer::GotoNextItem(int32 Distance)
-{
-	if (CurrentTimeLine.IsValid())
-	{
-		CurrentTimeLine.Pin()->MoveCursorByDistance(Distance);
-	}
-}
-
-void FLogVisualizer::GotoPreviousItem(int32 Distance)
-{
-	if (CurrentTimeLine.IsValid())
-	{
-		CurrentTimeLine.Pin()->MoveCursorByDistance(-Distance);
-	}
-}
-
-void FLogVisualizer::MoveCamera()
+FLogVisualizer& FLogVisualizer::Get()
 {
-	if (CurrentTimeLine.IsValid())
-	{
-		CurrentTimeLine.Pin()->UpdateCameraPosition();
-	}
+	return *StaticInstance;
 }
 
-
 FLinearColor FLogVisualizer::GetColorForCategory(int32 Index) const
 {
 	if (Index >= 0 && Index < sizeof(FLogVisualizerColorPalette) / sizeof(FLogVisualizerColorPalette[0]))
@@ -184,18 +142,196 @@ UWorld* FLogVisualizer::GetWorld(UObject* OptionalObject)
 	return World;
 }
 
-AActor* FLogVisualizer::GetVisualLoggerHelperActor()
+void FLogVisualizer::UpdateCameraPosition(FName RowName, int32 ItemIndes)
+{
+	const FVisualLoggerDBRow& DBRow = FVisualLoggerDatabase::Get().GetRowByName(RowName);
+	auto &Entries = DBRow.GetItems();
+	if (DBRow.GetCurrentItemIndex() == INDEX_NONE || Entries.IsValidIndex(DBRow.GetCurrentItemIndex()) == false)
+	{
+		return;
+	}
+
+	UWorld* World = GetWorld();
+	
+	FVector CurrentLocation = Entries[DBRow.GetCurrentItemIndex()].Entry.Location;
+
+	FVector Extent(150);
+	bool bFoundActor = false;
+	FName OwnerName = Entries[DBRow.GetCurrentItemIndex()].OwnerName;
+	for (FActorIterator It(World); It; ++It)
+	{
+		AActor* Actor = *It;
+		if (Actor->GetFName() == OwnerName)
+		{
+			FVector Orgin;
+			Actor->GetActorBounds(false, Orgin, Extent);
+			bFoundActor = true;
+			break;
+		}
+	}
+
+
+	const float DefaultCameraDistance = ULogVisualizerSettings::StaticClass()->GetDefaultObject<ULogVisualizerSettings>()->DefaultCameraDistance;
+	Extent = Extent.Size() < DefaultCameraDistance ? FVector(1) * DefaultCameraDistance : Extent;
+
+#if WITH_EDITOR
+	UEditorEngine *EEngine = Cast<UEditorEngine>(GEngine);
+	if (GIsEditor && EEngine != NULL)
+	{
+		for (auto ViewportClient : EEngine->AllViewportClients)
+		{
+			ViewportClient->FocusViewportOnBox(FBox::BuildAABB(CurrentLocation, Extent));
+		}
+	}
+	else if (AVisualLoggerCameraController::IsEnabled(World) && AVisualLoggerCameraController::Instance.IsValid() && AVisualLoggerCameraController::Instance->GetSpectatorPawn())
+	{
+		ULocalPlayer* LocalPlayer = Cast<ULocalPlayer>(AVisualLoggerCameraController::Instance->Player);
+		if (LocalPlayer && LocalPlayer->ViewportClient && LocalPlayer->ViewportClient->Viewport)
+		{
+
+			FViewport* Viewport = LocalPlayer->ViewportClient->Viewport;
+
+			FBox BoundingBox = FBox::BuildAABB(CurrentLocation, Extent);
+			const FVector Position = BoundingBox.GetCenter();
+			float Radius = BoundingBox.GetExtent().Size();
+			
+			FViewportCameraTransform ViewTransform;
+			ViewTransform.TransitionToLocation(Position, nullptr, true);
+
+			float NewOrthoZoom;
+			const float AspectRatio = 1.777777f;
+			uint32 MinAxisSize = (AspectRatio > 1.0f) ? Viewport->GetSizeXY().Y : Viewport->GetSizeXY().X;
+			float Zoom = Radius / (MinAxisSize / 2.0f);
+
+			NewOrthoZoom = Zoom * (Viewport->GetSizeXY().X*15.0f);
+			NewOrthoZoom = FMath::Clamp<float>(NewOrthoZoom, 250, MAX_FLT);
+			ViewTransform.SetOrthoZoom(NewOrthoZoom);
+
+			AVisualLoggerCameraController::Instance->GetSpectatorPawn()->TeleportTo(ViewTransform.GetLocation(), ViewTransform.GetRotation(), false, true);
+		}
+	}
+#endif
+}
+
+int32 FLogVisualizer::GetNextItem(FName RowName, int32 MoveDistance)
+{
+	FVisualLoggerDBRow& DBRow = FVisualLoggerDatabase::Get().GetRowByName(RowName);
+	int32 NewItemIndex = DBRow.GetCurrentItemIndex();
+
+	int32 Index = 0;
+	auto &Entries = DBRow.GetItems();
+	while (true)
+	{
+		NewItemIndex++;
+		if (Entries.IsValidIndex(NewItemIndex))
+		{
+			if (DBRow.IsItemVisible(NewItemIndex) == true && ++Index == MoveDistance)
+			{
+				break;
+			}
+		}
+		else
+		{
+			NewItemIndex = FMath::Clamp(NewItemIndex, 0, Entries.Num() - 1);
+			break;
+		}
+	}
+
+	return NewItemIndex;
+}
+
+int32 FLogVisualizer::GetPreviousItem(FName RowName, int32 MoveDistance)
+{
+	FVisualLoggerDBRow& DBRow = FVisualLoggerDatabase::Get().GetRowByName(RowName);
+	int32 NewItemIndex = DBRow.GetCurrentItemIndex();
+
+	int32 Index = 0;
+	auto &Entries = DBRow.GetItems();
+	while (true)
+	{
+		NewItemIndex--;
+		if (Entries.IsValidIndex(NewItemIndex))
+		{
+			if (DBRow.IsItemVisible(NewItemIndex) == true && ++Index == MoveDistance)
+			{
+				break;
+			}
+		}
+		else
+		{
+			NewItemIndex = FMath::Clamp(NewItemIndex, 0, Entries.Num() - 1);
+			break;
+		}
+	}
+	return NewItemIndex;
+}
+
+void FLogVisualizer::GotoNextItem(FName RowName, int32 MoveDistance)
 {
-	UWorld* World = FLogVisualizer::Get().GetWorld();
-	for (TActorIterator<AVisualLoggerRenderingActor> It(World); It; ++It)
+	FVisualLoggerDBRow& DBRow = FVisualLoggerDatabase::Get().GetRowByName(RowName);
+	const int32 NewItemIndex = GetNextItem(RowName, MoveDistance);
+
+	if (NewItemIndex != DBRow.GetCurrentItemIndex())
 	{
-		return *It;
+		auto &Entries = DBRow.GetItems();
+		float NewTimeStamp = Entries[NewItemIndex].Entry.TimeStamp;
 	}
+}
 
-	FActorSpawnParameters SpawnInfo;
-	SpawnInfo.SpawnCollisionHandlingOverride = ESpawnActorCollisionHandlingMethod::AlwaysSpawn;
-	SpawnInfo.Name = *FString::Printf(TEXT("VisualLoggerRenderingActor"));
-	AActor* HelperActor = World->SpawnActor<AVisualLoggerRenderingActor>(SpawnInfo);
+void FLogVisualizer::GotoPreviousItem(FName RowName, int32 MoveDistance)
+{
+	FVisualLoggerDBRow& DBRow = FVisualLoggerDatabase::Get().GetRowByName(RowName);
+	const int32 NewItemIndex = GetPreviousItem(RowName, MoveDistance);
 
-	return HelperActor;
+	if (NewItemIndex != DBRow.GetCurrentItemIndex())
+	{
+		auto &Entries = DBRow.GetItems();
+		float NewTimeStamp = Entries[NewItemIndex].Entry.TimeStamp;
+	}
 }
+
+void FLogVisualizer::GotoFirstItem(FName RowName)
+{
+	FVisualLoggerDBRow& DBRow = FVisualLoggerDatabase::Get().GetRowByName(RowName);
+	int32 NewItemIndex = DBRow.GetCurrentItemIndex();
+
+	auto &Entries = DBRow.GetItems();
+	for (int32 Index = 0; Index <= DBRow.GetCurrentItemIndex(); Index++)
+	{
+		if (DBRow.IsItemVisible(Index))
+		{
+			NewItemIndex = Index;
+			break;
+		}
+	}
+
+	if (NewItemIndex != DBRow.GetCurrentItemIndex())
+	{
+		//DBRow.MoveTo(NewItemIndex);
+		TimeSliderController->CommitScrubPosition(Entries[NewItemIndex].Entry.TimeStamp, false);
+	}
+}
+
+void FLogVisualizer::GotoLastItem(FName RowName)
+{
+	FVisualLoggerDBRow& DBRow = FVisualLoggerDatabase::Get().GetRowByName(RowName);
+	int32 NewItemIndex = DBRow.GetCurrentItemIndex();
+
+	auto &Entries = DBRow.GetItems();
+	for (int32 Index = Entries.Num() - 1; Index >= DBRow.GetCurrentItemIndex(); Index--)
+	{
+		if (DBRow.IsItemVisible(Index))
+		{
+			NewItemIndex = Index;
+			break;
+		}
+	}
+
+
+	if (NewItemIndex != DBRow.GetCurrentItemIndex())
+	{
+		//DBRow.MoveTo(NewItemIndex);
+		TimeSliderController->CommitScrubPosition(Entries[NewItemIndex].Entry.TimeStamp, false);
+	}
+}
+
diff --git a/Engine/Source/Developer/LogVisualizer/Private/LogVisualizer.h b/Engine/Source/Developer/LogVisualizer/Private/LogVisualizer.h
--- a/Engine/Source/Developer/LogVisualizer/Private/LogVisualizer.h
+++ b/Engine/Source/Developer/LogVisualizer/Private/LogVisualizer.h
@@ -17,47 +17,58 @@
 #include "VisualLogger/VisualLogger.h"
 #include "LogVisualizerSettings.h"
 #include "LogVisualizerSessionSettings.h"
+#include "VisualLoggerDatabase.h"
 
-DECLARE_DELEGATE_OneParam(FOnItemSelectionChanged, const FVisualLogDevice::FVisualLogEntryItem&);
-DECLARE_DELEGATE_OneParam(FOnObjectSelectionChanged, TSharedPtr<class STimeline>);
 DECLARE_DELEGATE_OneParam(FOnFiltersSearchChanged, const FText&);
-DECLARE_DELEGATE(FOnFiltersChanged);
-DECLARE_DELEGATE_ThreeParams(FOnLogLineSelectionChanged, TSharedPtr<FLogEntryItem> /*SelectedItem*/, int64 /*UserData*/, FName /*TagName*/);
+
+//DECLARE_DELEGATE(FOnFiltersChanged);
+DECLARE_MULTICAST_DELEGATE(FOnFiltersChanged);
+
+DECLARE_DELEGATE_ThreeParams(FOnLogLineSelectionChanged, TSharedPtr<struct FLogEntryItem> /*SelectedItem*/, int64 /*UserData*/, FName /*TagName*/);
 DECLARE_DELEGATE_RetVal_TwoParams(FReply, FOnKeyboardEvent, const FGeometry& /*MyGeometry*/, const FKeyEvent& /*InKeyEvent*/);
+DECLARE_DELEGATE_RetVal(float, FGetAnimationOutlinerFillPercentageFunc);
 
 struct FVisualLoggerEvents
 {
-	FVisualLoggerEvents() 
-	{
-	
-	}
-
-	FOnItemSelectionChanged OnItemSelectionChanged;
 	FOnFiltersChanged OnFiltersChanged;
-	FOnObjectSelectionChanged OnObjectSelectionChanged;
 	FOnLogLineSelectionChanged OnLogLineSelectionChanged;
 	FOnKeyboardEvent OnKeyboardEvent;
+	FGetAnimationOutlinerFillPercentageFunc GetAnimationOutlinerFillPercentageFunc;
 };
 
 class FVisualLoggerTimeSliderController;
 struct LOGVISUALIZER_API FLogVisualizer
 {
 	/** LogVisualizer interface*/
-	void Goto(float Timestamp, FName LogOwner = NAME_None);
-	void GotoNextItem( int32 Distance);
-	void GotoPreviousItem(int32 Distance);
-	void MoveCamera();
+	void Reset();
 
 	FLinearColor GetColorForCategory(int32 Index) const;
 	FLinearColor GetColorForCategory(const FString& InFilterName) const;
 	TSharedPtr<FVisualLoggerTimeSliderController> GetTimeSliderController() { return TimeSliderController; }
 	UWorld* GetWorld(UObject* OptionalObject = nullptr);
-	class AActor* GetVisualLoggerHelperActor();
-	FVisualLoggerEvents& GetVisualLoggerEvents() { return VisualLoggerEvents; }
+	FVisualLoggerEvents& GetEvents() { return VisualLoggerEvents; }
 
 	void SetCurrentVisualizer(TSharedPtr<class SVisualLogger> Visualizer) { CurrentVisualizer = Visualizer; }
 
-	void OnObjectSelectionChanged(TSharedPtr<class STimeline> TimeLine) { CurrentTimeLine = TimeLine; }
+	void SetAnimationOutlinerFillPercentage(float FillPercentage) { AnimationOutlinerFillPercentage = FillPercentage; }
+	float GetAnimationOutlinerFillPercentage() 
+	{
+		if (VisualLoggerEvents.GetAnimationOutlinerFillPercentageFunc.IsBound())
+		{
+			SetAnimationOutlinerFillPercentage(VisualLoggerEvents.GetAnimationOutlinerFillPercentageFunc.Execute());
+		}
+		return AnimationOutlinerFillPercentage; 
+	}
+
+	int32 GetNextItem(FName RowName, int32 MoveDistance = 1);
+	int32 GetPreviousItem(FName RowName, int32 MoveDistance = 1);
+
+	void GotoNextItem(FName RowName, int32 MoveDistance = 1);
+	void GotoPreviousItem(FName RowName, int32 MoveDistance = 1);
+	void GotoFirstItem(FName RowName);
+	void GotoLastItem(FName RowName);
+
+	void UpdateCameraPosition(FName Rowname, int32 ItemIndes);
 
 	/** Static access */
 	static void Initialize();
@@ -68,8 +79,8 @@ protected:
 	
 	TSharedPtr<FVisualLoggerTimeSliderController> TimeSliderController;
 	FVisualLoggerEvents VisualLoggerEvents;
-	TWeakPtr<class STimeline> CurrentTimeLine;
 	TWeakPtr<class SVisualLogger> CurrentVisualizer;
+	float AnimationOutlinerFillPercentage;
 };
 
 class SVisualLoggerTab : public SDockTab
@@ -78,7 +89,7 @@ public:
 	virtual bool SupportsKeyboardFocus() const override { return true; }
 	virtual FReply OnKeyDown(const FGeometry& MyGeometry, const FKeyEvent& InKeyEvent) override
 	{
-		return FLogVisualizer::Get().GetVisualLoggerEvents().OnKeyboardEvent.Execute(MyGeometry, InKeyEvent);
+		return FLogVisualizer::Get().GetEvents().OnKeyboardEvent.Execute(MyGeometry, InKeyEvent);
 	}
 };
 
@@ -88,7 +99,7 @@ public:
 	virtual bool SupportsKeyboardFocus() const override { return true; }
 	virtual FReply OnKeyDown(const FGeometry& MyGeometry, const FKeyEvent& InKeyEvent) override
 	{
-		return FLogVisualizer::Get().GetVisualLoggerEvents().OnKeyboardEvent.Execute(MyGeometry, InKeyEvent);
+		return FLogVisualizer::Get().GetEvents().OnKeyboardEvent.Execute(MyGeometry, InKeyEvent);
 	}
 };
 
diff --git a/Engine/Source/Developer/LogVisualizer/Private/LogVisualizerModule.cpp b/Engine/Source/Developer/LogVisualizer/Private/LogVisualizerModule.cpp
--- a/Engine/Source/Developer/LogVisualizer/Private/LogVisualizerModule.cpp
+++ b/Engine/Source/Developer/LogVisualizer/Private/LogVisualizerModule.cpp
@@ -26,10 +26,6 @@ public:
 	virtual void ShutdownModule() override;
 	// End IModuleInterface
 
-	virtual void Goto(float Timestamp, FName LogOwner) override;
-	virtual void GotoNextItem() override;
-	virtual void GotoPreviousItem() override;
-
 private:
 	TSharedRef<SDockTab> SpawnLogVisualizerTab(const FSpawnTabArgs& SpawnTabArgs);
 };
@@ -37,7 +33,9 @@ private:
 void FLogVisualizerModule::StartupModule()
 {
 	FLogVisualizerStyle::Initialize();
+	FVisualLoggerDatabase::Initialize();
 	FLogVisualizer::Initialize();
+	FVisualLoggerFilters::Initialize();
 
 	FVisualLoggerCommands::Register();
 	IModularFeatures::Get().RegisterModularFeature(VisualLoggerTabName, this);
@@ -73,7 +71,9 @@ void FLogVisualizerModule::ShutdownModule()
 		SettingsModule->UnregisterSettings("Editor", "General", "VisualLogger");
 	}
 
+	FVisualLoggerFilters::Shutdown();
 	FLogVisualizer::Shutdown();
+	FVisualLoggerDatabase::Shutdown();
 	FLogVisualizerStyle::Shutdown();
 }
 
@@ -91,21 +91,5 @@ TSharedRef<SDockTab> FLogVisualizerModule::SpawnLogVisualizerTab(const FSpawnTab
 	return MajorTab;
 }
 
-void FLogVisualizerModule::Goto(float Timestamp, FName LogOwner)
-{
-	FLogVisualizer::Get().Goto(Timestamp, LogOwner);
-}
-
-void FLogVisualizerModule::GotoNextItem()
-{
-	FLogVisualizer::Get().GotoNextItem(1);
-}
-
-void FLogVisualizerModule::GotoPreviousItem()
-{
-	FLogVisualizer::Get().GotoPreviousItem(1);
-}
-
-
 IMPLEMENT_MODULE(FLogVisualizerModule, LogVisualizer);
 #undef LOCTEXT_NAMESPACE
diff --git a/Engine/Source/Developer/LogVisualizer/Private/LogVisualizerSettings.cpp b/Engine/Source/Developer/LogVisualizer/Private/LogVisualizerSettings.cpp
--- a/Engine/Source/Developer/LogVisualizer/Private/LogVisualizerSettings.cpp
+++ b/Engine/Source/Developer/LogVisualizer/Private/LogVisualizerSettings.cpp
@@ -9,8 +9,6 @@
 #include "UnrealEdMisc.h"
 #endif // WITH_EDITOR
 
-FCategoryFiltersManager FCategoryFiltersManager::StaticManager;
-
 ULogVisualizerSettings::ULogVisualizerSettings(const FObjectInitializer& ObjectInitializer) 
 	: Super(ObjectInitializer)
 	, DebugMeshMaterialFakeLightName(TEXT("/Engine/EngineDebugMaterials/DebugMeshMaterialFakeLight.DebugMeshMaterialFakeLight"))
@@ -34,6 +32,39 @@ class UMaterial* ULogVisualizerSettings::GetDebugMeshMaterial()
 	return DebugMeshMaterialFakeLight;
 }
 
+void ULogVisualizerSettings::SavePresistentData()
+{
+	if (bPresistentFilters)
+	{
+		PresistentFilters = FVisualLoggerFilters::Get();
+	}
+	else
+	{
+		PresistentFilters = FVisualLoggerFilters();
+	}
+	SaveConfig();
+}
+
+void ULogVisualizerSettings::ClearPresistentData()
+{
+	if (bPresistentFilters)
+	{
+		PresistentFilters = FVisualLoggerFilters();
+	}
+}
+
+void ULogVisualizerSettings::LoadPresistentData()
+{
+	if (bPresistentFilters)
+	{
+		static_cast<FVisualLoggerFiltersData&>(FVisualLoggerFilters::Get()) = PresistentFilters;
+	}
+	else
+	{
+		FVisualLoggerFilters::Get() = FVisualLoggerFilters();
+	}
+}
+
 #if WITH_EDITOR
 void ULogVisualizerSettings::PostEditChangeProperty(struct FPropertyChangedEvent& PropertyChangedEvent)
 {
@@ -49,158 +80,200 @@ void ULogVisualizerSettings::PostEditChangeProperty(struct FPropertyChangedEvent
 }
 #endif
 
-bool FCategoryFiltersManager::MatchCategoryFilters(FString String, ELogVerbosity::Type Verbosity)
-{
-	const ULogVisualizerSettings* Settings = ULogVisualizerSettings::StaticClass()->GetDefaultObject<ULogVisualizerSettings>();
-	bool bFoundFilter = false;
-	for (const FCategoryFilter& Filter : Settings->CurrentFilters.Categories)
-	{
-		bFoundFilter = Filter.CategoryName == String;
-		if (bFoundFilter)
-		{
-			const bool bFineWithSearchString = Settings->bSearchInsideLogs == true ||
-				(Settings->bSearchInsideLogs == false && (Settings->CurrentFilters.SearchBoxFilter.Len() == 0 || Filter.CategoryName.Find(Settings->CurrentFilters.SearchBoxFilter) != INDEX_NONE));
-			return bFineWithSearchString && Verbosity <= Filter.LogVerbosity && Filter.Enabled;
-		}
-	}
-
-	return bFoundFilter;
-}
+//////////////////////////////////////////////////////////////////////////
+// FVisualLoggerFilters
+//////////////////////////////////////////////////////////////////////////
+TSharedPtr< struct FVisualLoggerFilters > FVisualLoggerFilters::StaticInstance;
 
-bool FCategoryFiltersManager::MatchObjectName(FString String)
+FVisualLoggerFilters& FVisualLoggerFilters::Get()
 {
-	const ULogVisualizerSettings* Settings = ULogVisualizerSettings::StaticClass()->GetDefaultObject<ULogVisualizerSettings>();
-	return (Settings->CurrentFilters.SelectedClasses.Num() == 0 || Settings->CurrentFilters.SelectedClasses.Find(String) != INDEX_NONE);
+	return *StaticInstance;
 }
 
-bool FCategoryFiltersManager::MatchSearchString(FString String)
+void FVisualLoggerFilters::Initialize()
 {
-	const ULogVisualizerSettings* Settings = ULogVisualizerSettings::StaticClass()->GetDefaultObject<ULogVisualizerSettings>();
-	return Settings->CurrentFilters.SearchBoxFilter == String;
-}
-
-void FCategoryFiltersManager::SetSearchString(FString InString) 
-{ 
-	ULogVisualizerSettings::StaticClass()->GetDefaultObject<ULogVisualizerSettings>()->CurrentFilters.SearchBoxFilter = InString; 
+	StaticInstance = MakeShareable(new FVisualLoggerFilters);
+	FVisualLoggerDatabase::Get().GetEvents().OnNewItem.AddRaw(StaticInstance.Get(), &FVisualLoggerFilters::OnNewItemHandler);
 }
 
-FString FCategoryFiltersManager::GetSearchString() 
-{ 
-	return ULogVisualizerSettings::StaticClass()->GetDefaultObject<ULogVisualizerSettings>()->CurrentFilters.SearchBoxFilter; 
+void FVisualLoggerFilters::Shutdown()
+{
+	FVisualLoggerDatabase::Get().GetEvents().OnNewItem.RemoveAll(StaticInstance.Get());
+	StaticInstance.Reset();
 }
 
-void FCategoryFiltersManager::SetObjectFilterString(FString InFilterString) 
-{ 
-	ULogVisualizerSettings::StaticClass()->GetDefaultObject<ULogVisualizerSettings>()->CurrentFilters.ObjectNameFilter = InFilterString; 
-}
+void FVisualLoggerFilters::OnNewItemHandler(const FVisualLoggerDBRow& DBRow, int32 ItemIndex)
+{
+	const FVisualLogDevice::FVisualLogEntryItem& Item = DBRow.GetItems()[ItemIndex];
+	TArray<FVisualLoggerCategoryVerbosityPair> Categories;
+	FVisualLoggerHelpers::GetCategories(Item.Entry, Categories);
+	for (auto& CategoryAndVerbosity : Categories)
+	{
+		AddCategory(CategoryAndVerbosity.CategoryName.ToString(), ELogVerbosity::All);
+	}
 
-FString FCategoryFiltersManager::GetObjectFilterString()
-{ 
-	return ULogVisualizerSettings::StaticClass()->GetDefaultObject<ULogVisualizerSettings>()->CurrentFilters.ObjectNameFilter; 
+	TMap<FString, TArray<FString> > OutCategories;
+	FVisualLoggerHelpers::GetHistogramCategories(Item.Entry, OutCategories);
+	for (const auto& CurrentCategory : OutCategories)
+	{
+		for (const auto& CurrentDataName : CurrentCategory.Value)
+		{
+			const FString GraphFilterName = CurrentCategory.Key + TEXT("$") + CurrentDataName;	
+			AddCategory(GraphFilterName, ELogVerbosity::All);
+		}
+	}
 }
 
-void FCategoryFiltersManager::AddCategory(FString InName, ELogVerbosity::Type InVerbosity)
+void FVisualLoggerFilters::AddCategory(FString InName, ELogVerbosity::Type InVerbosity)
 {
-	ULogVisualizerSettings* Settings = ULogVisualizerSettings::StaticClass()->GetDefaultObject<ULogVisualizerSettings>();
+	const int32 Num = Categories.Num();
+	for (int32 Index = 0; Index < Num; ++Index)
+	{
+		const FCategoryFilter& Filter = Categories[Index];
+		if (Filter.CategoryName == InName)
+		{
+			return;
+		}
+	}
+
 	FCategoryFilter Filter;
 	Filter.CategoryName = InName;
 	Filter.LogVerbosity = InVerbosity;
 	Filter.Enabled = true;
-	Settings->CurrentFilters.Categories.Add(Filter);
+	int32 Index = Categories.Add(Filter);
+
+	FastCategoryFilterMap.Reset(); // we need to recreate cache - pointers can be broken
+	FastCategoryFilterMap.Reserve(Categories.Num());
+	for (int32 Index = 0; Index < Categories.Num(); Index++)
+	{
+		FastCategoryFilterMap.Add(*Categories[Index].CategoryName, &Categories[Index]);
+	}
+
+	OnFilterCategoryAdded.Broadcast(InName, InVerbosity);
 }
 
-void FCategoryFiltersManager::RemoveCategory(FString InName)
+void FVisualLoggerFilters::RemoveCategory(FString InName)
 {
-	ULogVisualizerSettings* Settings = ULogVisualizerSettings::StaticClass()->GetDefaultObject<ULogVisualizerSettings>();
-	for (int32 Index = 0; Index < Settings->CurrentFilters.Categories.Num(); ++Index)
+	for (int32 Index = 0; Index < Categories.Num(); ++Index)
 	{
-		const FCategoryFilter& Filter = Settings->CurrentFilters.Categories[Index];
+		const FCategoryFilter& Filter = Categories[Index];
 		if (Filter.CategoryName == InName)
 		{
-			Settings->CurrentFilters.Categories.RemoveAt(Index);
+			Categories.RemoveAt(Index);
+			FastCategoryFilterMap.Remove(*InName);
 			break;
 		}
 	}
+
+	OnFilterCategoryRemoved.Broadcast(InName);
 }
 
-bool FCategoryFiltersManager::IsValidCategory(FString InName)
+FCategoryFilter& FVisualLoggerFilters::GetCategoryByName(const FName& InName)
 {
-	ULogVisualizerSettings* Settings = ULogVisualizerSettings::StaticClass()->GetDefaultObject<ULogVisualizerSettings>();
-	for (const FCategoryFilter& Filter : Settings->CurrentFilters.Categories)
+	FCategoryFilter* FilterPtr = FastCategoryFilterMap.Contains(InName) ? FastCategoryFilterMap[InName] : nullptr;
+	if (FilterPtr)
 	{
-		if (Filter.CategoryName == InName)
-		{
-			return true;
-		}
+		return *FilterPtr;
 	}
-	return Settings->CurrentFilters.Categories.Num() == 0;
+
+	static FCategoryFilter NoCategory;
+	return NoCategory;
 }
 
-FCategoryFilter& FCategoryFiltersManager::GetCategory(FString InName)
+FCategoryFilter& FVisualLoggerFilters::GetCategoryByName(const FString& InName)
 {
-	ULogVisualizerSettings* Settings = ULogVisualizerSettings::StaticClass()->GetDefaultObject<ULogVisualizerSettings>();
-	for (FCategoryFilter& Filter : Settings->CurrentFilters.Categories)
+	const FName NameAsName = *InName;
+	FCategoryFilter* FilterPtr = FastCategoryFilterMap.Contains(NameAsName) ? FastCategoryFilterMap[NameAsName] : nullptr;
+	if (FilterPtr)
 	{
-		if (Filter.CategoryName == InName)
-		{
-			return Filter;
-		}
+		return *FilterPtr;
 	}
 
 	static FCategoryFilter NoCategory;
 	return NoCategory;
 }
 
-void FCategoryFiltersManager::SelectObject(FString ObjectName)
+bool FVisualLoggerFilters::MatchObjectName(FString String)
 {
-	ULogVisualizerSettings* Settings = ULogVisualizerSettings::StaticClass()->GetDefaultObject<ULogVisualizerSettings>();
-	Settings->CurrentFilters.SelectedClasses.AddUnique(ObjectName);
+	return SelectedClasses.Num() == 0 || SelectedClasses.Find(String) != INDEX_NONE;
 }
 
-void FCategoryFiltersManager::RemoveObjectFromSelection(FString ObjectName)
+
+void FVisualLoggerFilters::SelectObject(FString ObjectName)
 {
-	ULogVisualizerSettings* Settings = ULogVisualizerSettings::StaticClass()->GetDefaultObject<ULogVisualizerSettings>();
-	Settings->CurrentFilters.SelectedClasses.Remove(ObjectName);
+	SelectedClasses.AddUnique(ObjectName);
+}
+
+void FVisualLoggerFilters::RemoveObjectFromSelection(FString ObjectName)
+{
+	SelectedClasses.Remove(ObjectName);
 }
 
-const TArray<FString>& FCategoryFiltersManager::GetSelectedObjects()
+const TArray<FString>& FVisualLoggerFilters::GetSelectedObjects() const
 {
-	return ULogVisualizerSettings::StaticClass()->GetDefaultObject<ULogVisualizerSettings>()->CurrentFilters.SelectedClasses;
+	return SelectedClasses;
 }
 
-void FCategoryFiltersManager::SavePresistentData()
+bool FVisualLoggerFilters::MatchCategoryFilters(FString String, ELogVerbosity::Type Verbosity)
 {
 	ULogVisualizerSettings* Settings = ULogVisualizerSettings::StaticClass()->GetDefaultObject<ULogVisualizerSettings>();
-	if (Settings->bPresistentFilters)
+	bool bFoundFilter = false;
+	for (const FCategoryFilter& Filter : Categories)
 	{
-		Settings->PresistentFilters = Settings->CurrentFilters;
+		bFoundFilter = Filter.CategoryName == String;
+		if (bFoundFilter)
+		{
+			const bool bFineWithSearchString = Settings->bSearchInsideLogs == true || (SearchBoxFilter.Len() == 0 || Filter.CategoryName.Find(SearchBoxFilter) != INDEX_NONE);
+			return bFineWithSearchString && Verbosity <= Filter.LogVerbosity && Filter.Enabled;
+		}
 	}
-	else
+
+	return bFoundFilter;
+}
+
+void FVisualLoggerFilters::DeactivateAllButThis(const FString& InName)
+{
+	for (FCategoryFilter& Filter : Categories)
 	{
-		Settings->PresistentFilters = FVisualLoggerFilters();
+		Filter.Enabled = Filter.CategoryName != InName ? false : true;
 	}
-	Settings->SaveConfig();
-
 }
 
-void FCategoryFiltersManager::ClearPresistentData()
+void FVisualLoggerFilters::EnableAllCategories()
 {
-	ULogVisualizerSettings* Settings = ULogVisualizerSettings::StaticClass()->GetDefaultObject<ULogVisualizerSettings>();
-	if (Settings->bPresistentFilters)
+	for (FCategoryFilter& Filter : Categories)
 	{
-		Settings->PresistentFilters = FVisualLoggerFilters();
+		Filter.Enabled = true;
 	}
 }
 
-void FCategoryFiltersManager::LoadPresistentData()
+
+void FVisualLoggerFilters::Reset()
 {
-	ULogVisualizerSettings* Settings = ULogVisualizerSettings::StaticClass()->GetDefaultObject<ULogVisualizerSettings>();
-	if (Settings->bPresistentFilters)
+	FastCategoryFilterMap.Reset();
+	SearchBoxFilter = FString();
+	ObjectNameFilter = FString();
+
+	Categories.Reset();
+	SelectedClasses.Reset();
+}
+
+void FVisualLoggerFilters::DisableGraphData(FName GraphName, FName DataName, bool SetAsDisabled)
+{
+	const FName FullName = *(GraphName.ToString() + TEXT("$") + DataName.ToString());
+	if (SetAsDisabled)
 	{
-		Settings->CurrentFilters = Settings->PresistentFilters;
+		DisabledGraphDatas.AddUnique(FullName);
 	}
 	else
 	{
-		Settings->CurrentFilters = FVisualLoggerFilters();
+		DisabledGraphDatas.RemoveSwap(FullName);
 	}
-}
\ No newline at end of file
+}
+
+
+bool FVisualLoggerFilters::IsGraphDataDisabled(FName GraphName, FName DataName)
+{
+	const FName FullName = *(GraphName.ToString() + TEXT("$") + DataName.ToString());
+	return DisabledGraphDatas.Find(FullName) != INDEX_NONE;
+}
diff --git a/Engine/Source/Developer/LogVisualizer/Private/LogVisualizerStyle.cpp b/Engine/Source/Developer/LogVisualizer/Private/LogVisualizerStyle.cpp
--- a/Engine/Source/Developer/LogVisualizer/Private/LogVisualizerStyle.cpp
+++ b/Engine/Source/Developer/LogVisualizer/Private/LogVisualizerStyle.cpp
@@ -35,6 +35,7 @@ FName FLogVisualizerStyle::GetStyleSetName()
 
 TSharedRef< FSlateStyleSet > FLogVisualizerStyle::Create()
 {
+	const FVector2D Icon8x8(8.0f, 8.0f);
 	const FVector2D Icon16x16(16.0f, 16.0f);
 	const FVector2D Icon20x20(20.0f, 20.0f);
 	const FVector2D Icon40x40(40.0f, 40.0f);
@@ -75,6 +76,42 @@ TSharedRef< FSlateStyleSet > FLogVisualizerStyle::Create()
 		Style.Set("ToolBar.Button.Hovered", new BOX_BRUSH("Common/RoundedSelection_16x", 4.0f / 16.0f, FLinearColor(0.728f, 0.364f, 0.003f)));
 
 	}
+	{
+		const FSlateColor SelectionColor = FLinearColor(0.728f, 0.364f, 0.003f);
+		const FSlateColor SelectionColor_Inactive = FLinearColor(0.25f, 0.25f, 0.25f);
+		const FSlateColor SelectorColor = FLinearColor(0.701f, 0.225f, 0.003f);
+		const FSlateColor DefaultForeground = FLinearColor(0.72f, 0.72f, 0.72f, 1.f);
+		const FSlateColor InvertedForeground = FLinearColor(0, 0, 0);
+
+		FTableRowStyle NormalTableRowStyle = FTableRowStyle()
+			.SetEvenRowBackgroundBrush(FSlateNoResource())
+			.SetEvenRowBackgroundHoveredBrush(IMAGE_BRUSH("Common/Selection", Icon8x8, FLinearColor(1.0f, 1.0f, 1.0f, 0.1f)))
+			.SetOddRowBackgroundBrush(FSlateNoResource())
+			.SetOddRowBackgroundHoveredBrush(IMAGE_BRUSH("Common/Selection", Icon8x8, FLinearColor(1.0f, 1.0f, 1.0f, 0.1f)))
+			.SetSelectorFocusedBrush(BORDER_BRUSH("Common/Selector", FMargin(4.f / 16.f), SelectorColor))
+			.SetActiveBrush(IMAGE_BRUSH("Common/Selection", Icon8x8, SelectionColor))
+			.SetActiveHoveredBrush(IMAGE_BRUSH("Common/Selection", Icon8x8, SelectionColor))
+			.SetInactiveBrush(IMAGE_BRUSH("Common/Selection", Icon8x8, SelectionColor_Inactive))
+			.SetInactiveHoveredBrush(IMAGE_BRUSH("Common/Selection", Icon8x8, SelectionColor_Inactive))
+			.SetTextColor(DefaultForeground)
+			.SetSelectedTextColor(InvertedForeground)
+			.SetDropIndicator_Above(BOX_BRUSH("Common/DropZoneIndicator_Above", FMargin(10.0f / 16.0f, 10.0f / 16.0f, 0, 0), SelectionColor))
+			.SetDropIndicator_Onto(BOX_BRUSH("Common/DropZoneIndicator_Onto", FMargin(4.0f / 16.0f), SelectionColor))
+			.SetDropIndicator_Below(BOX_BRUSH("Common/DropZoneIndicator_Below", FMargin(10.0f / 16.0f, 0, 0, 10.0f / 16.0f), SelectionColor));
+
+		Style.Set("TableView.DarkRow", FTableRowStyle(NormalTableRowStyle)
+			.SetEvenRowBackgroundBrush(IMAGE_BRUSH("Common/Selection", Icon8x8, FLinearColor(1.0f, 1.0f, 1.0f, 0.1f)))
+			.SetEvenRowBackgroundHoveredBrush(IMAGE_BRUSH("PropertyView/DetailCategoryMiddle_Hovered", FVector2D(16, 16)))
+			.SetOddRowBackgroundBrush(IMAGE_BRUSH("Common/Selection", Icon8x8, FLinearColor(1.0f, 1.0f, 1.0f, 0.1f)))
+			.SetOddRowBackgroundHoveredBrush(IMAGE_BRUSH("PropertyView/DetailCategoryMiddle_Hovered", FVector2D(16, 16)))
+			.SetSelectorFocusedBrush(BORDER_BRUSH("Common/Selector", FMargin(4.f / 16.f), SelectorColor))
+			.SetActiveBrush(IMAGE_BRUSH("Common/Selection", Icon8x8, SelectionColor))
+			.SetActiveHoveredBrush(IMAGE_BRUSH("Common/Selection", Icon8x8, SelectionColor))
+			.SetInactiveBrush(IMAGE_BRUSH("Common/Selection", Icon8x8, SelectionColor_Inactive))
+			.SetInactiveHoveredBrush(IMAGE_BRUSH("Common/Selection", Icon8x8, SelectionColor_Inactive))
+			.SetTextColor(FLinearColor(1.0f, 1.0f, 1.0f, 0.1f))
+			);
+	}
 
 	// Toolbar
 	{
diff --git a/Engine/Source/Developer/LogVisualizer/Private/SFilterWidget.cpp b/Engine/Source/Developer/LogVisualizer/Private/SFilterWidget.cpp
--- a/Engine/Source/Developer/LogVisualizer/Private/SFilterWidget.cpp
+++ b/Engine/Source/Developer/LogVisualizer/Private/SFilterWidget.cpp
@@ -16,6 +16,12 @@ void SFilterWidget::Construct(const FArguments& InArgs)
 	ColorCategory = InArgs._ColorCategory;
 	BorderBackgroundColor = FLinearColor(0.2f, 0.2f, 0.2f, 0.2f);
 
+	LastVerbosity = ELogVerbosity::NoLogging;
+	GetCaptionString(); // cache caption string
+
+	bWasEnabledLastTime = !IsEnabled();
+	GetTooltipString(); // cache tooltip string
+
 	// Get the tooltip and color of the type represented by this filter
 	FilterColor = ColorCategory;
 
@@ -51,30 +57,43 @@ void SFilterWidget::Construct(const FArguments& InArgs)
 FText SFilterWidget::GetCaptionString() const
 {
 	FString CaptionString;
-	FCategoryFilter& Filter = FCategoryFiltersManager::Get().GetCategory(GetFilterNameAsString());
-	const FString VerbosityStr = FOutputDevice::VerbosityToString((ELogVerbosity::Type)Filter.LogVerbosity);
-	if (Filter.LogVerbosity != ELogVerbosity::VeryVerbose)
-	{
-		CaptionString = FString::Printf(TEXT("%s [%s]"), *GetFilterNameAsString().Replace(TEXT("Log"), TEXT(""), ESearchCase::CaseSensitive), *VerbosityStr.Mid(0, 1));
-	}
-	else
+	FCategoryFilter& CategoryFilter = FVisualLoggerFilters::Get().GetCategoryByName(GetFilterNameAsString());
+	if (CategoryFilter.LogVerbosity != LastVerbosity)
 	{
-		CaptionString = FString::Printf(TEXT("%s [VV]"), *GetFilterNameAsString().Replace(TEXT("Log"), TEXT(""), ESearchCase::CaseSensitive));
+		const FString VerbosityStr = FOutputDevice::VerbosityToString((ELogVerbosity::Type)CategoryFilter.LogVerbosity);
+		if (CategoryFilter.LogVerbosity != ELogVerbosity::VeryVerbose)
+		{
+			CaptionString = FString::Printf(TEXT("%s [%s]"), *GetFilterNameAsString().Replace(TEXT("Log"), TEXT(""), ESearchCase::CaseSensitive), *VerbosityStr.Mid(0, 1));
+		}
+		else
+		{
+			CaptionString = FString::Printf(TEXT("%s [VV]"), *GetFilterNameAsString().Replace(TEXT("Log"), TEXT(""), ESearchCase::CaseSensitive));
+		}
+
+		CachedCaptionString = FText::FromString(CaptionString);
+		LastVerbosity = (ELogVerbosity::Type)CategoryFilter.LogVerbosity;
 	}
-	return FText::FromString(CaptionString);
+	return CachedCaptionString;
 }
 
 FText SFilterWidget::GetTooltipString() const
 {
-	FCategoryFilter& Filter = FCategoryFiltersManager::Get().GetCategory(GetFilterNameAsString());
-	const FString VerbosityStr = FOutputDevice::VerbosityToString((ELogVerbosity::Type)Filter.LogVerbosity);
-
-	return FText::FromString(
-		IsEnabled() ?
-		FString::Printf(TEXT("Enabled '%s' category for '%s' verbosity and lower\nRight click to change verbosity"), *GetFilterNameAsString(), *VerbosityStr)
-		:
-		FString::Printf(TEXT("Disabled '%s' category"), *GetFilterNameAsString())
-		);
+	if (bWasEnabledLastTime != IsEnabled())
+	{
+		FCategoryFilter& CategoryFilter = FVisualLoggerFilters::Get().GetCategoryByName(GetFilterNameAsString());
+		const FString VerbosityStr = FOutputDevice::VerbosityToString((ELogVerbosity::Type)CategoryFilter.LogVerbosity);
+
+		CachedTooltipString = FText::FromString(
+			IsEnabled() ?
+			FString::Printf(TEXT("Enabled '%s' category for '%s' verbosity and lower\nRight click to change verbosity"), *GetFilterNameAsString(), *VerbosityStr)
+			:
+			FString::Printf(TEXT("Disabled '%s' category"), *GetFilterNameAsString())
+			);
+
+		bWasEnabledLastTime = IsEnabled();
+	}
+
+	return CachedTooltipString;
 }
 
 TSharedRef<SWidget> SFilterWidget::GetRightClickMenuContent()
@@ -99,37 +118,63 @@ TSharedRef<SWidget> SFilterWidget::GetRightClickMenuContent()
 		}
 		MenuBuilder.EndSection();
 	}
+	MenuBuilder.BeginSection("FilterAction", LOCTEXT("FilterAction", "Context actions"));
+		MenuBuilder.AddMenuEntry(
+			LOCTEXT("DisableAllButThis", "Disable all but this"),
+			LOCTEXT("HideAllButThisTooltip", "Disable all other categories"),
+			FSlateIcon(),
+			FUIAction(FExecuteAction::CreateSP(this, &SFilterWidget::DisableAllButThis))
+			);
+		MenuBuilder.AddMenuEntry(
+			LOCTEXT("EnableAll", "Enable all categories"),
+			LOCTEXT("EnableAllTooltip", "Enable all categories"),
+			FSlateIcon(),
+			FUIAction(FExecuteAction::CreateSP(this, &SFilterWidget::EnableAllCategories))
+			);
+	MenuBuilder.EndSection();
 
 	return MenuBuilder.MakeWidget();
 }
 
 bool SFilterWidget::IsEnabled() const 
 { 
-	FCategoryFilter& Filter = FCategoryFiltersManager::Get().GetCategory(GetFilterNameAsString());
-	return Filter.Enabled;
+	FCategoryFilter& CategoryFilter = FVisualLoggerFilters::Get().GetCategoryByName(GetFilterNameAsString());
+	return CategoryFilter.Enabled;
 }
 
 void SFilterWidget::SetEnabled(bool InEnabled)
 {
-	FCategoryFilter& Filter = FCategoryFiltersManager::Get().GetCategory(GetFilterNameAsString());
-	if (InEnabled != Filter.Enabled)
+	FCategoryFilter& CategoryFilter = FVisualLoggerFilters::Get().GetCategoryByName(GetFilterNameAsString());
+	if (InEnabled != CategoryFilter.Enabled)
 	{
-		Filter.Enabled = InEnabled;
+		CategoryFilter.Enabled = InEnabled;
 		OnFilterChanged.ExecuteIfBound();
 	}
 }
 
 void SFilterWidget::FilterToggled(ECheckBoxState NewState)
 {
-	FCategoryFilter& Filter = FCategoryFiltersManager::Get().GetCategory(GetFilterNameAsString());
 	SetEnabled(NewState == ECheckBoxState::Checked);
 }
 
 void SFilterWidget::SetVerbosityFilter(int32 SelectedVerbosityIndex)
 {
-	FCategoryFilter& Filter = FCategoryFiltersManager::Get().GetCategory(GetFilterNameAsString());
-	Filter.LogVerbosity = (ELogVerbosity::Type)SelectedVerbosityIndex;
+	FCategoryFilter& CategoryFilter = FVisualLoggerFilters::Get().GetCategoryByName(GetFilterNameAsString());
+	CategoryFilter.LogVerbosity = (ELogVerbosity::Type)SelectedVerbosityIndex;
+	OnFilterChanged.ExecuteIfBound();
+}
+
+void SFilterWidget::DisableAllButThis()
+{
+	FVisualLoggerFilters::Get().DeactivateAllButThis(GetFilterNameAsString());
 	OnFilterChanged.ExecuteIfBound();
 }
 
+void SFilterWidget::EnableAllCategories()
+{
+	FVisualLoggerFilters::Get().EnableAllCategories();
+	OnFilterChanged.ExecuteIfBound();
+}
+
+
 #undef LOCTEXT_NAMESPACE
diff --git a/Engine/Source/Developer/LogVisualizer/Private/SFilterWidget.h b/Engine/Source/Developer/LogVisualizer/Private/SFilterWidget.h
--- a/Engine/Source/Developer/LogVisualizer/Private/SFilterWidget.h
+++ b/Engine/Source/Developer/LogVisualizer/Private/SFilterWidget.h
@@ -86,6 +86,10 @@ private:
 	/** Removes this filter from the filter list */
 	void SetVerbosityFilter(int32 SelectedVerbosityIndex);
 
+	void DisableAllButThis();
+
+	void EnableAllCategories();
+
 	void RemoveFilter()
 	{
 		TSharedRef<SFilterWidget> Self = SharedThis(this);
@@ -154,4 +158,11 @@ private:
 
 	/** Default color for border background */
 	FLinearColor BorderBackgroundColor;
+
+	/** hacky to cache data */
+	mutable bool bWasEnabledLastTime;
+	mutable ELogVerbosity::Type LastVerbosity;
+
+	mutable FText CachedCaptionString;
+	mutable FText CachedTooltipString;
 };
diff --git a/Engine/Source/Developer/LogVisualizer/Private/STimeline.cpp b/Engine/Source/Developer/LogVisualizer/Private/STimeline.cpp
--- a/Engine/Source/Developer/LogVisualizer/Private/STimeline.cpp
+++ b/Engine/Source/Developer/LogVisualizer/Private/STimeline.cpp
@@ -32,12 +32,6 @@ FReply STimelineLabelAnchor::OnMouseButtonDown(const FGeometry& MyGeometry, cons
 	return FReply::Unhandled();
 }
 
-
-STimeline::~STimeline()
-{
-	ULogVisualizerSettings::StaticClass()->GetDefaultObject<ULogVisualizerSettings>()->OnSettingChanged().RemoveAll(this);
-}
-
 FReply STimeline::OnMouseButtonDown(const FGeometry& MyGeometry, const FPointerEvent& MouseEvent)
 {
 	Owner->ChangeSelection(SharedThis(this), MouseEvent);
@@ -45,57 +39,42 @@ FReply STimeline::OnMouseButtonDown(const FGeometry& MyGeometry, const FPointerE
 	return FReply::Unhandled();
 }
 
-FReply STimeline::OnMouseButtonUp(const FGeometry& MyGeometry, const FPointerEvent& MouseEvent)
+FReply STimeline::OnKeyDown(const FGeometry& MyGeometry, const FKeyEvent& InKeyEvent)
 {
-	return FReply::Unhandled();
-}
-
-bool STimeline::IsSelected() const
-{
-	// Ask the tree if we are selected
-	return Owner->IsNodeSelected(SharedThis(const_cast<STimeline*>(this)));
-}
-
-void STimeline::OnSelect()
-{
-	if (TimelineBar.IsValid())
+	FReply ReturnValue = FReply::Unhandled();
+	if (IsSelected() == false)
 	{
-		TimelineBar->OnSelect();
+		return ReturnValue;
 	}
+
+	return FLogVisualizer::Get().GetEvents().OnKeyboardEvent.Execute(MyGeometry, InKeyEvent);
 }
 
-void STimeline::OnDeselect()
+bool STimeline::IsSelected() const
 {
-	if (TimelineBar.IsValid())
-	{
-		TimelineBar->OnDeselect();
-	}
+	// Ask the tree if we are selected
+	return FVisualLoggerDatabase::Get().IsRowSelected(GetName());
 }
 
 void STimeline::UpdateVisibility()
 {
 	ULogVisualizerSettings* Settings = ULogVisualizerSettings::StaticClass()->GetDefaultObject<ULogVisualizerSettings>();
+
+	FVisualLoggerDBRow& DataRow = FVisualLoggerDatabase::Get().GetRowByName(GetName());
+	const TArray<FVisualLogDevice::FVisualLogEntryItem>& Entries = DataRow.GetItems();
+
 	const bool bJustIgnore = Settings->bIgnoreTrivialLogs && Entries.Num() <= Settings->TrivialLogsThreshold;
-	const bool bVisibleByOwnerClasses = FCategoryFiltersManager::Get().MatchObjectName(OwnerClassName.ToString());
-	const bool bIsCollapsed = bJustIgnore || HiddenEntries.Num() == Entries.Num() || (SearchFilter.Len() > 0 && Name.ToString().Find(SearchFilter) == INDEX_NONE);
+	const bool bVisibleByOwnerClasses = FVisualLoggerFilters::Get().MatchObjectName(OwnerClassName.ToString());
+	const bool bIsCollapsed = bJustIgnore || DataRow.GetNumberOfHiddenItems() == Entries.Num() || (SearchFilter.Len() > 0 && OwnerName.ToString().Find(SearchFilter) == INDEX_NONE);
 
-	SetVisibility(bIsCollapsed || !bVisibleByOwnerClasses ? EVisibility::Collapsed : EVisibility::Visible);
+	const bool bIsHidden = bIsCollapsed || !bVisibleByOwnerClasses;
+	SetVisibility(bIsHidden ? EVisibility::Collapsed : EVisibility::Visible);
 	if (bIsCollapsed)
 	{
 		Owner->SetSelectionState(SharedThis(this), false, false);
 	}
-}
-
-void STimeline::UpdateVisibilityForItems()
-{
-	HiddenEntries.Reset();
-	ULogVisualizerSettings* Settings = ULogVisualizerSettings::StaticClass()->GetDefaultObject<ULogVisualizerSettings>();
-	const FString QuickSearchStrng = FCategoryFiltersManager::Get().GetSearchString(); 
 
-	for (auto& CurrentEntry : Entries)
-	{
-		UpdateVisibilityForEntry(CurrentEntry, QuickSearchStrng, Settings->bSearchInsideLogs);
-	}
+	FVisualLoggerDatabase::Get().SetRowVisibility(GetName(), !bIsHidden);
 }
 
 void STimeline::OnFiltersSearchChanged(const FText& Filter)
@@ -105,7 +84,6 @@ void STimeline::OnFiltersSearchChanged(const FText& Filter)
 
 void STimeline::OnFiltersChanged()
 {
-	UpdateVisibilityForItems();
 	UpdateVisibility();
 }
 
@@ -115,89 +93,38 @@ void STimeline::OnSearchChanged(const FText& Filter)
 	UpdateVisibility();
 }
 
-bool STimeline::IsEntryHidden(const FVisualLogDevice::FVisualLogEntryItem& EntryItem) const
-{
-	return HiddenEntries.Contains(&EntryItem);
-}
-
 void STimeline::HandleLogVisualizerSettingChanged(FName InName)
 {
 	UpdateVisibility();
 }
 
-void STimeline::UpdateVisibilityForEntry(FVisualLogDevice::FVisualLogEntryItem& CurrentEntry, const FString& SearchString, bool bSearchInsideLogs)
+const TArray<FVisualLogDevice::FVisualLogEntryItem>& STimeline::GetEntries()
 {
-	TArray<FVisualLoggerCategoryVerbosityPair> OutCategories;
-	FVisualLoggerHelpers::GetCategories(CurrentEntry.Entry, OutCategories);
-	bool bHasValidCategories = false;
-	for (FVisualLoggerCategoryVerbosityPair& Categoryair : OutCategories)
-	{
-		bHasValidCategories = bHasValidCategories || FCategoryFiltersManager::Get().MatchCategoryFilters(Categoryair.CategoryName.ToString(), Categoryair.Verbosity);
-	}
-
-	if (bSearchInsideLogs && bHasValidCategories && SearchString.Len() > 0)
-	{
-		bool bMatchSearchString = false;
-		for (const FVisualLogLine& CurrentLine : CurrentEntry.Entry.LogLines)
-		{
-			if (CurrentLine.Line.Find(SearchString) != INDEX_NONE || CurrentLine.Category.ToString().Find(SearchString) != INDEX_NONE)
-			{
-				bMatchSearchString = true;
-				break;
-			}
-		}
-		if (!bMatchSearchString)
-		{
-			for (const FVisualLogEvent& CurrentEvent : CurrentEntry.Entry.Events)
-			{
-				if (CurrentEvent.Name.Find(SearchString) != INDEX_NONE)
-				{
-					bMatchSearchString = true;
-					break;
-				}
-			}
-		}
-
-		if (!bMatchSearchString)
-		{
-			HiddenEntries.AddUnique(&CurrentEntry);
-		}
-	}
-	else
-	{
-		if (bHasValidCategories == false)
-		{
-			HiddenEntries.AddUnique(&CurrentEntry);
-		}
-	}
+	FVisualLoggerDBRow& DataRow = FVisualLoggerDatabase::Get().GetRowByName(GetName());
+	return DataRow.GetItems();
 }
 
 void STimeline::AddEntry(const FVisualLogDevice::FVisualLogEntryItem& Entry) 
 { 
-	Entries.Add(Entry); 
+	//Entries.Add(Entry); 
 
 	ULogVisualizerSettings* Settings = ULogVisualizerSettings::StaticClass()->GetDefaultObject<ULogVisualizerSettings>();
 	
-	UpdateVisibilityForEntry(Entries[Entries.Num() - 1], FCategoryFiltersManager::Get().GetSearchString(), Settings->bSearchInsideLogs);
+	FVisualLoggerDBRow& DataRow = FVisualLoggerDatabase::Get().GetRowByName(GetName());
+	const TArray<FVisualLogDevice::FVisualLogEntryItem>& Entries = DataRow.GetItems();
+
 	UpdateVisibility();
 
-	if (Settings->bStickToRecentData && IsSelected())
-	{
-		TimelineBar->SnapScrubPosition(Entry.Entry.TimeStamp);
-	}
 }
 
-void STimeline::Construct(const FArguments& InArgs, TSharedPtr<SVisualLoggerView> VisualLoggerView, TSharedPtr<FVisualLoggerTimeSliderController> TimeSliderController, TSharedPtr<STimelinesContainer> InContainer, const FVisualLogDevice::FVisualLogEntryItem& Entry)
+void STimeline::Construct(const FArguments& InArgs, TSharedPtr<FVisualLoggerTimeSliderController> TimeSliderController, TSharedPtr<STimelinesContainer> InContainer, FName InName, FName InOwnerClassName)
 {
 	OnGetMenuContent = InArgs._OnGetMenuContent;
 
 	Owner = InContainer;
-	Name = Entry.OwnerName;
-	OwnerClassName = Entry.OwnerClassName;
+	OwnerName = InName;
+	OwnerClassName = InOwnerClassName;
 	
-	ULogVisualizerSettings::StaticClass()->GetDefaultObject<ULogVisualizerSettings>()->OnSettingChanged().AddRaw(this, &STimeline::HandleLogVisualizerSettingChanged);
-
-
 	ChildSlot
 		[
 			SNew(SHorizontalBox)
@@ -205,7 +132,7 @@ void STimeline::Construct(const FArguments& InArgs, TSharedPtr<SVisualLoggerView
 			.Padding(FMargin(0, 4, 0, 0))
 			.HAlign(HAlign_Fill)
 			.VAlign(VAlign_Fill)
-			.FillWidth(TAttribute<float>(VisualLoggerView.Get(), &SVisualLoggerView::GetAnimationOutlinerFillPercentage))
+			.FillWidth(TAttribute<float>::Create(TAttribute<float>::FGetter::CreateLambda([=] { return FLogVisualizer::Get().GetAnimationOutlinerFillPercentage(); })))
 			[
 				SAssignNew(PopupAnchor, STimelineLabelAnchor, SharedThis(this))
 				.OnGetMenuContent(OnGetMenuContent)
@@ -228,7 +155,7 @@ void STimeline::Construct(const FArguments& InArgs, TSharedPtr<SVisualLoggerView
 							.VAlign(VAlign_Fill)
 							[
 								SNew(STextBlock)
-								.Text( FText::FromName(Name) )
+								.Text(FText::FromName(OwnerName))
 								.ShadowOffset(FVector2D(1.f, 1.f))
 							]
 							+ SVerticalBox::Slot()
@@ -258,7 +185,36 @@ void STimeline::Construct(const FArguments& InArgs, TSharedPtr<SVisualLoggerView
 			]
 		];
 
-	AddEntry(Entry);
+	ULogVisualizerSettings::StaticClass()->GetDefaultObject<ULogVisualizerSettings>()->OnSettingChanged().AddRaw(this, &STimeline::HandleLogVisualizerSettingChanged);
+	FVisualLoggerDatabase::Get().GetEvents().OnNewItem.AddRaw(this, &STimeline::OnNewItemHandler);
+	FVisualLoggerDatabase::Get().GetEvents().OnRowSelectionChanged.AddRaw(this, &STimeline::OnRowSelectionChanged);
+}
+
+STimeline::~STimeline()
+{
+	if (ULogVisualizerSettings::StaticClass() && ULogVisualizerSettings::StaticClass()->GetDefaultObject<ULogVisualizerSettings>())
+	{
+		ULogVisualizerSettings::StaticClass()->GetDefaultObject<ULogVisualizerSettings>()->OnSettingChanged().RemoveAll(this);
+	}
+	FVisualLoggerDatabase::Get().GetEvents().OnNewItem.RemoveAll(this);
+	FVisualLoggerDatabase::Get().GetEvents().OnRowSelectionChanged.RemoveAll(this);
+}
+
+void STimeline::OnRowSelectionChanged(const TArray<FName>& RowNames)
+{
+	if (RowNames.Num() > 0)
+	{
+		const FName RowName = RowNames[RowNames.Num() - 1];
+	}
+}
+
+void STimeline::OnNewItemHandler(const FVisualLoggerDBRow& BDRow, int32 ItemIndex)
+{
+	const FVisualLogDevice::FVisualLogEntryItem& Entry = BDRow.GetItems()[ItemIndex];
+	if (GetName() == Entry.OwnerName)
+	{
+		AddEntry(Entry);
+	}
 }
 
 const FSlateBrush* STimeline::GetBorder() const
@@ -275,34 +231,28 @@ const FSlateBrush* STimeline::GetBorder() const
 
 void STimeline::Goto(float ScrubPosition) 
 { 
-	TimelineBar->SnapScrubPosition(ScrubPosition); 
+
 }
 
 void STimeline::GotoNextItem() 
 { 
-	TimelineBar->GotoNextItem(); 
+	FLogVisualizer::Get().GotoNextItem(GetName());
 }
 
 void STimeline::GotoPreviousItem() 
 { 
-	TimelineBar->GotoPreviousItem(); 
+	FLogVisualizer::Get().GotoPreviousItem(GetName());
 }
 
 void STimeline::MoveCursorByDistance(int32 Distance)
 {
 	if (Distance > 0)
 	{
-		TimelineBar->GotoNextItem(FMath::Abs(Distance));
+		FLogVisualizer::Get().GotoNextItem(GetName(), FMath::Abs(Distance));
 	}
 	else
 	{
-		TimelineBar->GotoPreviousItem(FMath::Abs(Distance));
+		FLogVisualizer::Get().GotoPreviousItem(GetName(), FMath::Abs(Distance));
 	}
 	
 }
-
-
-void STimeline::UpdateCameraPosition()
-{
-	TimelineBar->UpdateCameraPosition();
-}
diff --git a/Engine/Source/Developer/LogVisualizer/Private/STimeline.h b/Engine/Source/Developer/LogVisualizer/Private/STimeline.h
--- a/Engine/Source/Developer/LogVisualizer/Private/STimeline.h
+++ b/Engine/Source/Developer/LogVisualizer/Private/STimeline.h
@@ -9,51 +9,52 @@ class STimeline : public SCompoundWidget
 {
 public:
 	SLATE_BEGIN_ARGS(STimeline){}
-		SLATE_EVENT(FOnItemSelectionChanged, OnItemSelectionChanged)
 		SLATE_EVENT(FOnGetContent, OnGetMenuContent)
 	SLATE_END_ARGS();
 
 	virtual FReply OnMouseButtonDown(const FGeometry& MyGeometry, const FPointerEvent& MouseEvent) override;
-	virtual FReply OnMouseButtonUp(const FGeometry& MyGeometry, const FPointerEvent& MouseEvent) override;
+	virtual FReply OnKeyDown(const FGeometry& MyGeometry, const FKeyEvent& InKeyEvent) override;
+	virtual bool SupportsKeyboardFocus() const override { return true; }
 
-	virtual void OnSelect();
-	virtual void OnDeselect();
+	virtual TOptional<bool> OnQueryShowFocus(const EFocusCause InFocusCause) const override
+	{
+		return TOptional<bool>(false);
+	}
 
 	/** Constructs this widget with InArgs */
 	virtual ~STimeline();
-	void Construct(const FArguments& InArgs, TSharedPtr<class SVisualLoggerView> VisualLoggerView, TSharedPtr<class FVisualLoggerTimeSliderController> InTimeSliderController, TSharedPtr<class STimelinesContainer> InContainer, const FVisualLogDevice::FVisualLogEntryItem& Entry);
+	void Construct(const FArguments& InArgs, TSharedPtr<class FVisualLoggerTimeSliderController> InTimeSliderController, TSharedPtr<class STimelinesContainer> InContainer, FName Name, FName OwnerClassName);
 	bool IsSelected() const;
-	bool IsEntryHidden(const FVisualLogDevice::FVisualLogEntryItem&) const;
 	const FSlateBrush* GetBorder() const;
 
 	void OnFiltersChanged();
 	void OnSearchChanged(const FText& Filter);
 	void OnFiltersSearchChanged(const FText& Filter);
 	void UpdateVisibility();
-	void UpdateVisibilityForItems();
 
 	void AddEntry(const FVisualLogDevice::FVisualLogEntryItem& Entry);
-	const TArray<FVisualLogDevice::FVisualLogEntryItem>& GetEntries() { return Entries; }
-	FName GetName() { return Name; }
-	FName GetOwnerClassName() { return OwnerClassName; }
+	const TArray<FVisualLogDevice::FVisualLogEntryItem>& GetEntries();
+	FName GetName() const { return OwnerName; }
+	FName GetOwnerClassName() const { return OwnerClassName; }
 	void HandleLogVisualizerSettingChanged(FName Name);
-	void UpdateVisibilityForEntry(FVisualLogDevice::FVisualLogEntryItem& Entry, const FString& SearchString, bool SearchInsideLogs);
 	TSharedPtr<class STimelinesContainer> GetOwner() { return Owner; }
 
 	void Goto(float ScrubPosition);
 	void GotoNextItem();
 	void GotoPreviousItem();
 	void MoveCursorByDistance(int32 Distance);
-	void UpdateCameraPosition();
 
 protected:
-	TArray<FVisualLogDevice::FVisualLogEntryItem> Entries;
-	TArray<const FVisualLogDevice::FVisualLogEntryItem*> HiddenEntries;
+	void OnNewItemHandler(const FVisualLoggerDBRow& BDRow, int32 ItemIndex);
+	void OnRowSelectionChanged(const TArray<FName>& RowNames);
+
+protected:
+	//TArray<FName> HiddenEntries;
 	TSharedPtr<class STimelinesContainer> Owner;
 	TSharedPtr<class STimelineBar> TimelineBar;
 	TSharedPtr<SMenuAnchor> PopupAnchor;
 
-	FName Name;
+	FName OwnerName;
 	FName OwnerClassName;
 	FString SearchFilter;
 
diff --git a/Engine/Source/Developer/LogVisualizer/Private/STimelineBar.cpp b/Engine/Source/Developer/LogVisualizer/Private/STimelineBar.cpp
--- a/Engine/Source/Developer/LogVisualizer/Private/STimelineBar.cpp
+++ b/Engine/Source/Developer/LogVisualizer/Private/STimelineBar.cpp
@@ -4,7 +4,6 @@
 #include "STimelineBar.h"
 #include "TimeSliderController.h"
 #include "STimelinesContainer.h"
-#include "VisualLoggerCameraController.h"
 #if WITH_EDITOR
 #	include "Editor/UnrealEd/Public/EditorComponents.h"
 #	include "Editor/UnrealEd/Public/EditorReimportHandler.h"
@@ -18,7 +17,20 @@
 FReply STimelineBar::OnMouseButtonDown(const FGeometry& MyGeometry, const FPointerEvent& MouseEvent)
 {
 	TimelineOwner.Pin()->OnMouseButtonDown(MyGeometry, MouseEvent);
-	return TimeSliderController->OnMouseButtonDown(SharedThis(this), MyGeometry, MouseEvent);
+	FReply Replay = TimeSliderController->OnMouseButtonDown(SharedThis(this), MyGeometry, MouseEvent);
+	if (Replay.IsEventHandled())
+	{
+		FName RowName = TimelineOwner.Pin()->GetName();
+		FVisualLoggerDBRow& DBRow = FVisualLoggerDatabase::Get().GetRowByName(RowName);
+		const float ScrubPosition = TimeSliderController->GetTimeSliderArgs().ScrubPosition.Get();
+		const int32 ClosestItem = DBRow.GetClosestItem(ScrubPosition);
+		const auto& Items = DBRow.GetItems();
+		if (Items.IsValidIndex(ClosestItem))
+		{
+			TimeSliderController->CommitScrubPosition(Items[ClosestItem].Entry.TimeStamp, false);
+		}
+	}
+	return Replay;
 }
 
 FReply STimelineBar::OnMouseButtonUp(const FGeometry& MyGeometry, const FPointerEvent& MouseEvent)
@@ -26,11 +38,18 @@ FReply STimelineBar::OnMouseButtonUp(const FGeometry& MyGeometry, const FPointer
 	TimelineOwner.Pin()->OnMouseButtonUp(MyGeometry, MouseEvent);
 
 	FReply  Replay = TimeSliderController->OnMouseButtonUp(SharedThis(this), MyGeometry, MouseEvent);
-	if (MouseEvent.GetEffectingButton() == EKeys::LeftMouseButton)
+	if (Replay.IsEventHandled())
 	{
-		SnapScrubPosition(TimeSliderController->GetTimeSliderArgs().ScrubPosition.Get());
+		FName RowName = TimelineOwner.Pin()->GetName();
+		FVisualLoggerDBRow& DBRow = FVisualLoggerDatabase::Get().GetRowByName(RowName);
+		const float ScrubPosition = TimeSliderController->GetTimeSliderArgs().ScrubPosition.Get();
+		const int32 ClosestItem = DBRow.GetClosestItem(ScrubPosition);
+		const auto& Items = DBRow.GetItems();
+		if (Items.IsValidIndex(ClosestItem))
+		{
+			TimeSliderController->CommitScrubPosition(Items[ClosestItem].Entry.TimeStamp, false);
+		}
 	}
-
 	return Replay;
 }
 
@@ -39,275 +58,28 @@ FReply STimelineBar::OnMouseMove(const FGeometry& MyGeometry, const FPointerEven
 	return TimeSliderController->OnMouseMove(SharedThis(this), MyGeometry, MouseEvent);
 }
 
-void STimelineBar::UpdateCameraPosition()
-{
-	auto &Entries = TimelineOwner.Pin()->GetEntries();
-	UWorld* World = FLogVisualizer::Get().GetWorld();
-
-	FVector CurrentLocation = Entries[CurrentItemIndex].Entry.Location;
-
-	FVector Extent(150);
-	bool bFoundActor = false;
-	FName OwnerName = Entries[CurrentItemIndex].OwnerName;
-	for (FActorIterator It(World); It; ++It)
-	{
-		AActor* Actor = *It;
-		if (Actor->GetFName() == OwnerName)
-		{
-			FVector Orgin;
-			Actor->GetActorBounds(false, Orgin, Extent);
-			bFoundActor = true;
-			break;
-		}
-	}
-
-
-	const float DefaultCameraDistance = ULogVisualizerSettings::StaticClass()->GetDefaultObject<ULogVisualizerSettings>()->DefaultCameraDistance;
-	Extent = Extent.Size() < DefaultCameraDistance ? FVector(1) * DefaultCameraDistance : Extent;
-
-#if WITH_EDITOR
-	UEditorEngine *EEngine = Cast<UEditorEngine>(GEngine);
-	if (GIsEditor && EEngine != NULL)
-	{
-		for (auto ViewportClient : EEngine->AllViewportClients)
-		{
-			ViewportClient->FocusViewportOnBox(FBox::BuildAABB(CurrentLocation, Extent));
-		}
-	}
-	else if (AVisualLoggerCameraController::IsEnabled(World) && AVisualLoggerCameraController::Instance.IsValid() && AVisualLoggerCameraController::Instance->GetSpectatorPawn())
-	{
-		ULocalPlayer* LocalPlayer = Cast<ULocalPlayer>(AVisualLoggerCameraController::Instance->Player);
-		if (LocalPlayer && LocalPlayer->ViewportClient && LocalPlayer->ViewportClient->Viewport)
-		{
-			FViewport* Viewport = LocalPlayer->ViewportClient->Viewport;
-
-			FBox BoundingBox = FBox::BuildAABB(CurrentLocation, Extent);
-			const FVector Position = BoundingBox.GetCenter();
-			float Radius = BoundingBox.GetExtent().Size();
-
-			FViewportCameraTransform ViewTransform;
-			ViewTransform.TransitionToLocation(Position, SharedThis(this), true);
-
-			float NewOrthoZoom;
-			const float AspectRatio = 1.777777f;
-			uint32 MinAxisSize = (AspectRatio > 1.0f) ? Viewport->GetSizeXY().Y : Viewport->GetSizeXY().X;
-			float Zoom = Radius / (MinAxisSize / 2.0f);
-
-			NewOrthoZoom = Zoom * (Viewport->GetSizeXY().X*15.0f);
-			NewOrthoZoom = FMath::Clamp<float>(NewOrthoZoom, 250, MAX_FLT);
-			ViewTransform.SetOrthoZoom(NewOrthoZoom);
-
-			AVisualLoggerCameraController::Instance->GetSpectatorPawn()->TeleportTo(ViewTransform.GetLocation(), ViewTransform.GetRotation(), false, true);
-		}
-	}
-#endif
-}
-
 FReply STimelineBar::OnMouseButtonDoubleClick(const FGeometry& MyGeometry, const FPointerEvent& MouseEvent)
 {
-	auto &Entries = TimelineOwner.Pin()->GetEntries();
-	SnapScrubPosition(TimeSliderController->GetTimeSliderArgs().ScrubPosition.Get());
+	FName RowName = TimelineOwner.Pin()->GetName();
+	FVisualLoggerDBRow& DBRow = FVisualLoggerDatabase::Get().GetRowByName(RowName);
 	UWorld* World = FLogVisualizer::Get().GetWorld();
-	if (World && MouseEvent.GetEffectingButton() == EKeys::LeftMouseButton && Entries.IsValidIndex(CurrentItemIndex))
+	if (World && MouseEvent.GetEffectingButton() == EKeys::LeftMouseButton)
 	{
-		UpdateCameraPosition();
+		FLogVisualizer::Get().UpdateCameraPosition(RowName, DBRow.GetCurrentItemIndex());
 		return FReply::Handled();
 	}
 
 	return FReply::Unhandled();
 }
 
-void STimelineBar::GotoNextItem(int32 MoveDistance)
-{
-	int32 NewItemIndex = CurrentItemIndex;
-	auto &Entries = TimelineOwner.Pin()->GetEntries();
-	int32 Index = 0;
-	while (true)
-	{
-		NewItemIndex++;
-		if (Entries.IsValidIndex(NewItemIndex))
-		{
-			auto& CurrentEntryItem = Entries[NewItemIndex];
-			if (TimelineOwner.Pin()->IsEntryHidden(CurrentEntryItem) == false && ++Index == MoveDistance)
-			{
-				break;
-			}
-		}
-		else
-		{
-			NewItemIndex = FMath::Clamp(NewItemIndex, 0, Entries.Num() - 1);
-			break;
-		}
-	}
-
-	if (NewItemIndex != CurrentItemIndex)
-	{
-		float NewTimeStamp = Entries[NewItemIndex].Entry.TimeStamp;
-		SnapScrubPosition(NewItemIndex);
-	}
-}
-
-void STimelineBar::GotoPreviousItem(int32 MoveDistance)
+STimelineBar::~STimelineBar()
 {
-	int32 NewItemIndex = CurrentItemIndex;
-	auto &Entries = TimelineOwner.Pin()->GetEntries();
-	int32 Index = 0;
-	while (true)
-	{
-		NewItemIndex--;
-		if (Entries.IsValidIndex(NewItemIndex))
-		{
-			auto& CurrentEntryItem = Entries[NewItemIndex];
-			if (TimelineOwner.Pin()->IsEntryHidden(CurrentEntryItem) == false && ++Index == MoveDistance)
-			{
-				break;
-			}
-		}
-		else
-		{
-			NewItemIndex = FMath::Clamp(NewItemIndex, 0, Entries.Num() - 1);
-			break;
-		}
-	}
-
-	if (NewItemIndex != CurrentItemIndex)
-	{
-		float NewTimeStamp = Entries[NewItemIndex].Entry.TimeStamp;
-		SnapScrubPosition(NewItemIndex);
-	}
-}
-
-FReply STimelineBar::OnKeyDown(const FGeometry& MyGeometry, const FKeyEvent& InKeyEvent)
-{
-	FReply ReturnValue = FReply::Unhandled();
-	if (TimelineOwner.Pin()->IsSelected() == false)
-	{
-		return ReturnValue;
-	}
-
-	if (CurrentItemIndex != INDEX_NONE)
-	{
-		TRange<float> LocalViewRange = TimeSliderController->GetTimeSliderArgs().ViewRange.Get();
-		float RangeSize = LocalViewRange.Size<float>();
-
-		const FKey Key = InKeyEvent.GetKey();
-		int32 NewItemIndex = CurrentItemIndex;
-		auto &Entries = TimelineOwner.Pin()->GetEntries();
-		if (Key == EKeys::Home)
-		{
-			for (int32 Index = 0; Index < Entries.Num(); Index++)
-			{
-				if (TimelineOwner.Pin()->IsEntryHidden(Entries[Index]) == false)
-				{
-					NewItemIndex = Index;
-					break;
-				}
-			}
-			if (NewItemIndex != CurrentItemIndex)
-			{
-				SnapScrubPosition(NewItemIndex);
-			}
-			ReturnValue = FReply::Handled();
-		}
-		else if (Key == EKeys::End)
-		{
-			for (int32 Index = Entries.Num()-1; Index >= 0; Index--)
-			{
-				if (TimelineOwner.Pin()->IsEntryHidden(Entries[Index]) == false)
-				{
-					NewItemIndex = Index;
-					break;
-				}
-			}
-			if (NewItemIndex != CurrentItemIndex)
-			{
-				SnapScrubPosition(NewItemIndex);
-			}
-			ReturnValue = FReply::Handled();
-		}
-		else if (Key == EKeys::Left)
-		{
-			GotoPreviousItem(InKeyEvent.IsLeftControlDown() ? InKeyEvent.IsLeftShiftDown() ? 20 : 10 : 1);
-			ReturnValue = FReply::Handled();
-		}
-		else if (Key == EKeys::Right)
-		{
-			GotoNextItem(InKeyEvent.IsLeftControlDown() ? InKeyEvent.IsLeftShiftDown() ? 20 : 10 : 1);
-			ReturnValue = FReply::Handled();
-		}
-	}
-
-	return ReturnValue;
-}
-
-int32 STimelineBar::GetClosestItem(float Time) const
-{
-	int32 BestItemIndex = INDEX_NONE;
-	float BestDistance = MAX_FLT;
-	auto &Entries = TimelineOwner.Pin()->GetEntries();
-	for (int32 Index = 0; Index < Entries.Num(); Index++)
-	{
-		auto& CurrentEntryItem = Entries[Index];
-
-		if (TimelineOwner.Pin()->IsEntryHidden(CurrentEntryItem))
-		{
-			continue;
-		}
-		TArray<FVisualLoggerCategoryVerbosityPair> OutCategories;
-		const float CurrentDist = FMath::Abs(CurrentEntryItem.Entry.TimeStamp - Time);
-		if (CurrentDist < BestDistance)
-		{
-			BestDistance = CurrentDist;
-			BestItemIndex = Index;
-		}
-	}
-
-	const float CurrentDist = Entries.IsValidIndex(CurrentItemIndex) ? FMath::Abs(Entries[CurrentItemIndex].Entry.TimeStamp - Time) : MAX_FLT;
-
-	if (BestItemIndex != INDEX_NONE && CurrentDist > BestDistance)
-	{
-		return BestItemIndex;
-	}
-
-	return CurrentItemIndex;
-}
-
-void STimelineBar::SnapScrubPosition(float ScrubPosition)
-{
-	uint32 BestItemIndex = GetClosestItem(ScrubPosition);
-	if (BestItemIndex != INDEX_NONE)
-	{
-		auto &Entries = TimelineOwner.Pin()->GetEntries();
-		const float CurrentTime = Entries[BestItemIndex].Entry.TimeStamp;
-		if (CurrentItemIndex != BestItemIndex)
-		{
-			CurrentItemIndex = BestItemIndex;
-			FLogVisualizer::Get().GetVisualLoggerEvents().OnItemSelectionChanged.ExecuteIfBound(Entries[BestItemIndex]);
-		}
-		TimeSliderController->CommitScrubPosition(CurrentTime, false);
-	}
-}
-
-void STimelineBar::SnapScrubPosition(int32 NewItemIndex)
-{
-	if (NewItemIndex != INDEX_NONE)
-	{
-		auto &Entries = TimelineOwner.Pin()->GetEntries();
-		const float CurrentTime = Entries[NewItemIndex].Entry.TimeStamp;
-		if (CurrentItemIndex != NewItemIndex)
-		{
-			CurrentItemIndex = NewItemIndex;
-			FLogVisualizer::Get().GetVisualLoggerEvents().OnItemSelectionChanged.ExecuteIfBound(Entries[NewItemIndex]);
-		}
-		TimeSliderController->CommitScrubPosition(CurrentTime, false);
-	}
 }
 
 void STimelineBar::Construct(const FArguments& InArgs, TSharedPtr<FVisualLoggerTimeSliderController> InTimeSliderController, TSharedPtr<STimeline> InTimelineOwner)
 {
 	TimeSliderController = InTimeSliderController;
 	TimelineOwner = InTimelineOwner;
-	CurrentItemIndex = INDEX_NONE;
 
 	TRange<float> LocalViewRange = TimeSliderController->GetTimeSliderArgs().ViewRange.Get();
 }
@@ -317,17 +89,6 @@ FVector2D STimelineBar::ComputeDesiredSize(float) const
 	return FVector2D(5000.0f, 20.0f);
 }
 
-void STimelineBar::OnSelect()
-{
-	FSlateApplication::Get().SetKeyboardFocus(SharedThis(this), EFocusCause::SetDirectly);
-	CurrentItemIndex = INDEX_NONE;
-}
-
-void STimelineBar::OnDeselect()
-{
-	CurrentItemIndex = INDEX_NONE;
-}
-
 int32 STimelineBar::OnPaint(const FPaintArgs& Args, const FGeometry& AllottedGeometry, const FSlateRect& MyClippingRect, FSlateWindowElementList& OutDrawElements, int32 LayerId, const FWidgetStyle& InWidgetStyle, bool bParentEnabled) const
 {
 	//@TODO: Optimize it like it was with old LogVisualizer, to draw everything much faster (SebaK)
@@ -366,9 +127,10 @@ int32 STimelineBar::OnPaint(const FPaintArgs& Args, const FGeometry& AllottedGeo
 
 	TArray<float> ErrorTimes;
 	TArray<float> WarningTimes;
-	auto &Entries = TimelineOwner.Pin()->GetEntries();
 	int32 EntryIndex = 0;
 
+	FVisualLoggerDBRow& DBRow = FVisualLoggerDatabase::Get().GetRowByName(TimelineOwner.Pin()->GetName());
+	auto &Entries = DBRow.GetItems();
 	while (EntryIndex < Entries.Num())
 	{
 		const FVisualLogEntry& Entry = Entries[EntryIndex].Entry;
@@ -378,7 +140,7 @@ int32 STimelineBar::OnPaint(const FPaintArgs& Args, const FGeometry& AllottedGeo
 			continue;
 		}
 
-		if (TimelineOwner.Pin()->IsEntryHidden(Entries[EntryIndex]))
+		if (DBRow.IsItemVisible(EntryIndex)==false)
 		{
 			EntryIndex++;
 			continue;
@@ -397,7 +159,7 @@ int32 STimelineBar::OnPaint(const FPaintArgs& Args, const FGeometry& AllottedGeo
 				break;
 			}
 
-			if (TimelineOwner.Pin()->IsEntryHidden(Entries[StartIndex]))
+			if (DBRow.IsItemVisible(StartIndex) == false)
 			{
 				continue;
 			}
@@ -485,18 +247,12 @@ int32 STimelineBar::OnPaint(const FPaintArgs& Args, const FGeometry& AllottedGeo
 			);
 	}
 
-	int32 BestItemIndex = GetClosestItem(LocalScrubPosition);
+	int32 BestItemIndex = DBRow.GetClosestItem(LocalScrubPosition);
 
-	if (BestItemIndex != INDEX_NONE && TimelineOwner.Pin()->IsSelected())
+	if (TimelineOwner.Pin()->IsSelected() && DBRow.GetCurrentItemIndex() != INDEX_NONE)
 	{
-		const auto &BestItemEntry = TimelineOwner.Pin()->GetEntries()[BestItemIndex];
-		if (BestItemIndex != CurrentItemIndex)
-		{
-			FLogVisualizer::Get().GetVisualLoggerEvents().OnItemSelectionChanged.ExecuteIfBound(BestItemEntry);
-			CurrentItemIndex = BestItemIndex;
-		}
-
-		float CurrentTime = BestItemEntry.Entry.TimeStamp;
+		const auto &HighlightedItemEntry = DBRow.GetCurrentItem();
+		float CurrentTime = HighlightedItemEntry.Entry.TimeStamp;
 		float LinePos = (CurrentTime - LocalViewRange.GetLowerBoundValue()) * PixelsPerInput;
 
 		FSlateDrawElement::MakeBox(
@@ -513,4 +269,4 @@ int32 STimelineBar::OnPaint(const FPaintArgs& Args, const FGeometry& AllottedGeo
 	}
 
 	return RetLayerId;
-}
+}
\ No newline at end of file
diff --git a/Engine/Source/Developer/LogVisualizer/Private/STimelineBar.h b/Engine/Source/Developer/LogVisualizer/Private/STimelineBar.h
--- a/Engine/Source/Developer/LogVisualizer/Private/STimelineBar.h
+++ b/Engine/Source/Developer/LogVisualizer/Private/STimelineBar.h
@@ -10,7 +10,6 @@ class STimelineBar : public SLeafWidget
 {
 public:
 	SLATE_BEGIN_ARGS(STimelineBar){}
-		SLATE_EVENT(FOnItemSelectionChanged, OnItemSelectionChanged)
 	SLATE_END_ARGS()
 
 	/** Constructs this widget with InArgs */
@@ -18,25 +17,15 @@ public:
 	virtual FReply OnMouseButtonDown(const FGeometry& MyGeometry, const FPointerEvent& MouseEvent) override;
 	virtual FReply OnMouseButtonUp(const FGeometry& MyGeometry, const FPointerEvent& MouseEvent) override;
 	virtual FReply OnMouseMove(const FGeometry& MyGeometry, const FPointerEvent& MouseEvent) override;
-	virtual FReply OnKeyDown(const FGeometry& MyGeometry, const FKeyEvent& InKeyEvent) override;
 
-	virtual bool SupportsKeyboardFocus() const override { return true;  }
+	virtual bool SupportsKeyboardFocus() const override { return false;  }
 
+	virtual ~STimelineBar();
 	void Construct(const FArguments& InArgs, TSharedPtr<FVisualLoggerTimeSliderController> InTimeSliderController, TSharedPtr<class STimeline> TimelineOwner);
 	virtual int32 OnPaint(const FPaintArgs& Args, const FGeometry& AllottedGeometry, const FSlateRect& MyClippingRect, FSlateWindowElementList& OutDrawElements, int32 LayerId, const FWidgetStyle& InWidgetStyle, bool bParentEnabled) const override;
 	FVector2D ComputeDesiredSize(float) const override;
-	void SnapScrubPosition(float ScrubPosition);
-	void SnapScrubPosition(int32 NewItemIndex);
-	int32 GetClosestItem(float Time) const;
-
-	void OnSelect();
-	void OnDeselect();
-	void GotoNextItem(int32 MoveDistance = 1);
-	void GotoPreviousItem(int32 MoveDistance = 1);
-	void UpdateCameraPosition();
 
 protected:
 	TSharedPtr<class FVisualLoggerTimeSliderController> TimeSliderController;
 	TWeakPtr<class STimeline> TimelineOwner;
-	mutable int32 CurrentItemIndex;
 };
diff --git a/Engine/Source/Developer/LogVisualizer/Private/STimelinesContainer.cpp b/Engine/Source/Developer/LogVisualizer/Private/STimelinesContainer.cpp
--- a/Engine/Source/Developer/LogVisualizer/Private/STimelinesContainer.cpp
+++ b/Engine/Source/Developer/LogVisualizer/Private/STimelinesContainer.cpp
@@ -9,22 +9,6 @@
 
 #define LOCTEXT_NAMESPACE "STimelinesContainer"
 
-TSharedRef<SWidget> STimelinesContainer::MakeTimeline(TSharedPtr<class SVisualLoggerView> InVisualLoggerView, TSharedPtr<class FVisualLoggerTimeSliderController> InTimeSliderController, const FVisualLogDevice::FVisualLogEntryItem& Entry)
-{
-	TSharedPtr<STimeline> NewTimeline;
-
-	ContainingBorder->AddSlot()
-		[
-			SAssignNew(NewTimeline, STimeline, InVisualLoggerView, InTimeSliderController, SharedThis(this), Entry)
-			.OnItemSelectionChanged(FLogVisualizer::Get().GetVisualLoggerEvents().OnItemSelectionChanged)
-			.OnGetMenuContent(this, &STimelinesContainer::GetRightClickMenuContent)
-		];
-
-	TimelineItems.Add(NewTimeline.ToSharedRef());
-
-	return NewTimeline.ToSharedRef();
-}
-
 TSharedRef<SWidget> STimelinesContainer::GetRightClickMenuContent()
 {
 	FMenuBuilder MenuBuilder(true, NULL);
@@ -58,33 +42,39 @@ TSharedRef<SWidget> STimelinesContainer::GetRightClickMenuContent()
 
 void STimelinesContainer::SetSelectionState(TSharedPtr<class STimeline> AffectedNode, bool bSelect, bool bDeselectOtherNodes)
 {
-	if (bSelect)
+	const FName RowName = AffectedNode->GetName();
+	const bool bIsSelected = FVisualLoggerDatabase::Get().IsRowSelected(RowName);
+	if (bSelect && (!bIsSelected || bDeselectOtherNodes))
 	{
-		if (bDeselectOtherNodes)
-		{
-			// empty current selection set unless multiple selecting
-			SelectedNodes.Empty();
-		}
+		FVisualLoggerDatabase::Get().SelectRow(RowName, bDeselectOtherNodes);
+	}
+	else if (!bSelect && bIsSelected && AffectedNode.IsValid())
+	{
+		FVisualLoggerDatabase::Get().DeselectRow(RowName);
+	}
+}
+
 
-		if (AffectedNode.IsValid())
+void STimelinesContainer::OnObjectSelectionChanged(const TArray<FName>& RowNames)
+{
+	CachedSelectedTimelines.Reset();
+	for (TSharedPtr<class STimeline> CurrentTimeline : TimelineItems)
+	{
+		if (RowNames.Find(CurrentTimeline->GetName()) != INDEX_NONE)
 		{
-			SelectedNodes.Add(AffectedNode);
-			AffectedNode->OnSelect();
+			CachedSelectedTimelines.Add(CurrentTimeline);
 		}
-		FLogVisualizer::Get().GetVisualLoggerEvents().OnObjectSelectionChanged.ExecuteIfBound(AffectedNode);
 	}
-	else if (AffectedNode.IsValid())
+
+	if (CachedSelectedTimelines.Num() >= 1)
 	{
-		// Not selecting so remove the node from the selection set
-		SelectedNodes.Remove(AffectedNode);
-		AffectedNode->OnDeselect();
-		FLogVisualizer::Get().GetVisualLoggerEvents().OnObjectSelectionChanged.ExecuteIfBound(AffectedNode);
+		FSlateApplication::Get().SetKeyboardFocus(SharedThis(CachedSelectedTimelines[CachedSelectedTimelines.Num()-1].Get()), EFocusCause::Navigation);
 	}
 }
 
 bool STimelinesContainer::IsNodeSelected(TSharedPtr<class STimeline> Node) const
 {
-	return SelectedNodes.Contains(Node);
+	return FVisualLoggerDatabase::Get().IsRowSelected(Node->GetName());// SelectedNodes.Contains(Node);
 }
 
 void STimelinesContainer::ChangeSelection(class TSharedPtr<class STimeline> Timeline, const FPointerEvent& MouseEvent)
@@ -102,12 +92,12 @@ void STimelinesContainer::ChangeSelection(class TSharedPtr<class STimeline> Time
 	}
 	else
 	{
-		if (SelectedNodes.Num() == 0 && TimelineItems.Num())
+		if (CachedSelectedTimelines.Num() == 0 && TimelineItems.Num())
 		{
 			SetSelectionState(TimelineItems[0], true, true);
 		}
 
-		TSharedPtr<class STimeline> LastSelected = SelectedNodes.Num() ? SelectedNodes[SelectedNodes.Num() - 1] : nullptr;
+		TSharedPtr<class STimeline> LastSelected = CachedSelectedTimelines.Num() ? CachedSelectedTimelines[CachedSelectedTimelines.Num() - 1] : nullptr;
 		if (LastSelected.IsValid())
 		{
 			bool bStartedSelection = false;
@@ -166,9 +156,7 @@ FReply STimelinesContainer::OnMouseWheel(const FGeometry& MyGeometry, const FPoi
 {
 	if (MouseEvent.IsLeftControlDown() || MouseEvent.IsLeftShiftDown())
 	{
-		FReply RetValue = TimeSliderController->OnMouseWheel(SharedThis(this), MyGeometry, MouseEvent);
-		FLogVisualizer::Get().GetVisualLoggerEvents().OnFiltersChanged.ExecuteIfBound();
-		return RetValue;
+		return TimeSliderController->OnMouseWheel(SharedThis(this), MyGeometry, MouseEvent);
 	}
 	return FReply::Unhandled();
 }
@@ -184,12 +172,12 @@ FReply STimelinesContainer::OnKeyDown(const FGeometry& MyGeometry, const FKeyEve
 
 		return FReply::Handled();
 	}
-	else if (InKeyEvent.GetKey() == EKeys::Platform_Delete && SelectedNodes.Num() > 0)
+	else if (InKeyEvent.GetKey() == EKeys::Platform_Delete && CachedSelectedTimelines.Num() > 0)
 	{
 		TWeakPtr<class STimeline>  NotSelectedOne;
-		for (auto & CurrentNode : SelectedNodes)
+		for (auto & CurrentNode : CachedSelectedTimelines)
 		{
-			TSharedPtr<class STimeline> LastSelected = SelectedNodes[SelectedNodes.Num() - 1];
+			TSharedPtr<class STimeline> LastSelected = CachedSelectedTimelines[CachedSelectedTimelines.Num() - 1];
 			bool bFoundSelectedOne = false;
 			for (auto& CurrentItem : TimelineItems)
 			{
@@ -206,6 +194,7 @@ FReply STimelinesContainer::OnKeyDown(const FGeometry& MyGeometry, const FKeyEve
 					bFoundSelectedOne = true;
 				}
 			}
+			FVisualLoggerDatabase::Get().RemoveRow(CurrentNode->GetName());
 			TimelineItems.Remove(CurrentNode);
 			ContainingBorder->RemoveSlot(CurrentNode.ToSharedRef());
 		}
@@ -216,7 +205,7 @@ FReply STimelinesContainer::OnKeyDown(const FGeometry& MyGeometry, const FKeyEve
 	else if (InKeyEvent.GetKey() == EKeys::Up || InKeyEvent.GetKey() == EKeys::Down)
 	{
 		TSharedPtr<class STimeline> PreviousTimeline;
-		TSharedPtr<class STimeline> LastSelected = SelectedNodes[SelectedNodes.Num() - 1];
+		TSharedPtr<class STimeline> LastSelected = CachedSelectedTimelines[CachedSelectedTimelines.Num() - 1];
 		for (int32 Index = 0; Index < TimelineItems.Num(); ++Index)
 		{
 			auto& CurrentItem = TimelineItems[Index];
@@ -257,12 +246,19 @@ FReply STimelinesContainer::OnKeyDown(const FGeometry& MyGeometry, const FKeyEve
 
 void STimelinesContainer::ResetData()
 {
-	for (auto CurrentItem : TimelineItems)
+	for (auto& CurrentItem : TimelineItems)
 	{
 		ContainingBorder->RemoveSlot(CurrentItem.ToSharedRef());
+		CurrentItem.Reset();
 	}
 	TimelineItems.Reset();
 
+	for (auto& CurrentItem : CachedSelectedTimelines)
+	{
+		CurrentItem.Reset();
+	}
+	CachedSelectedTimelines.Reset();
+
 	CachedMinTime = FLT_MAX;
 	CachedMaxTime = 0;
 	TimeSliderController->SetClampRange(0, 5);
@@ -288,34 +284,36 @@ void STimelinesContainer::Construct(const FArguments& InArgs, TSharedRef<class S
 	CachedMinTime = FLT_MAX;
 	CachedMaxTime = 0;
 
+	FVisualLoggerDatabase::Get().GetEvents().OnNewRow.AddRaw(this, &STimelinesContainer::OnNewRowHandler);
+	FVisualLoggerDatabase::Get().GetEvents().OnNewItem.AddRaw(this, &STimelinesContainer::OnNewItemHandler);
+	FVisualLoggerDatabase::Get().GetEvents().OnRowSelectionChanged.AddRaw(this, &STimelinesContainer::OnObjectSelectionChanged);
+	FVisualLoggerDatabase::Get().GetEvents().OnRowChangedVisibility.AddRaw(this, &STimelinesContainer::OnRowChangedVisibility);
 }
 
-void STimelinesContainer::OnSearchChanged(const FText& Filter)
+STimelinesContainer::~STimelinesContainer()
 {
-	for (auto CurrentItem : TimelineItems)
-	{
-		CurrentItem->OnSearchChanged(Filter);
-	}
+	FVisualLoggerDatabase::Get().GetEvents().OnNewRow.RemoveAll(this);
+	FVisualLoggerDatabase::Get().GetEvents().OnNewItem.RemoveAll(this);
+	FVisualLoggerDatabase::Get().GetEvents().OnRowSelectionChanged.RemoveAll(this);
+	FVisualLoggerDatabase::Get().GetEvents().OnRowChangedVisibility.RemoveAll(this);
 }
 
-void STimelinesContainer::OnNewLogEntry(const FVisualLogDevice::FVisualLogEntryItem& Entry)
+void STimelinesContainer::OnNewRowHandler(const FVisualLoggerDBRow& DBRow)
 {
-	bool bCreateNew = true;
-	for (auto CurrentItem : TimelineItems)
-	{
-		if (CurrentItem->GetName() == Entry.OwnerName)
-		{
-			CurrentItem->AddEntry(Entry);
-			bCreateNew = false;
-			break;
-		}
-	}
+	TSharedPtr<STimeline> NewTimeline;
 
-	if (bCreateNew)
-	{
-		MakeTimeline(VisualLoggerView, TimeSliderController, Entry);
-	}
+	ContainingBorder->AddSlot()
+		[
+			SAssignNew(NewTimeline, STimeline, TimeSliderController, SharedThis(this), DBRow.GetOwnerName(), DBRow.GetOwnerClassName())
+			.OnGetMenuContent(this, &STimelinesContainer::GetRightClickMenuContent)
+		];
+
+	TimelineItems.Add(NewTimeline.ToSharedRef());
+}
 
+void STimelinesContainer::OnNewItemHandler(const FVisualLoggerDBRow& BDRow, int32 ItemIndex)
+{
+	const FVisualLogDevice::FVisualLogEntryItem& Entry = BDRow.GetItems()[ItemIndex];
 	CachedMinTime = CachedMinTime < Entry.Entry.TimeStamp ? CachedMinTime : Entry.Entry.TimeStamp;
 	CachedMaxTime = CachedMaxTime > Entry.Entry.TimeStamp ? CachedMaxTime : Entry.Entry.TimeStamp;
 
@@ -324,9 +322,17 @@ void STimelinesContainer::OnNewLogEntry(const FVisualLogDevice::FVisualLogEntryI
 	float ZoomLevel = LocalViewRange.Size<float>() / ClampRange.Size<float>();
 
 	TimeSliderController->GetTimeSliderArgs().ClampRange = TRange<float>(CachedMinTime, CachedMaxTime + 0.1f);
-	if ( FMath::Abs(ZoomLevel - 1) <= SMALL_NUMBER)
-	{
+	//if ( FMath::Abs(ZoomLevel - 1) <= SMALL_NUMBER)
+	//{
 		TimeSliderController->GetTimeSliderArgs().ViewRange = TimeSliderController->GetTimeSliderArgs().ClampRange;
+	//}
+}
+
+void STimelinesContainer::OnSearchChanged(const FText& Filter)
+{
+	for (auto CurrentItem : TimelineItems)
+	{
+		CurrentItem->OnSearchChanged(Filter);
 	}
 }
 
@@ -352,10 +358,24 @@ void STimelinesContainer::GenerateReport()
 		.ClientSize(FVector2D(720, 768))
 		.Title(NSLOCTEXT("LogVisualizerReport", "WindowTitle", "Log Visualizer Report"))
 		[
-			SNew(SVisualLoggerReport, SelectedNodes, VisualLoggerView)
+			SNew(SVisualLoggerReport, CachedSelectedTimelines, VisualLoggerView)
 		];
 
 	FSlateApplication::Get().AddWindow(NewWindow);
 }
 
+
+void STimelinesContainer::OnRowChangedVisibility(const FName& InName)
+{
+	for (TSharedPtr<class STimeline> CurrentTimeline : TimelineItems)
+	{
+		if (InName == CurrentTimeline->GetName())
+		{
+			CurrentTimeline->SetVisibility(FVisualLoggerDatabase::Get().IsRowVisible(InName) ? EVisibility::Visible : EVisibility::Collapsed);
+			break;
+		}
+	}
+	
+}
+
 #undef LOCTEXT_NAMESPACE
diff --git a/Engine/Source/Developer/LogVisualizer/Private/STimelinesContainer.h b/Engine/Source/Developer/LogVisualizer/Private/STimelinesContainer.h
--- a/Engine/Source/Developer/LogVisualizer/Private/STimelinesContainer.h
+++ b/Engine/Source/Developer/LogVisualizer/Private/STimelinesContainer.h
@@ -11,6 +11,8 @@ public:
 	SLATE_BEGIN_ARGS(STimelinesContainer){}
 	SLATE_END_ARGS()
 
+	virtual ~STimelinesContainer();
+
 	virtual FReply OnMouseButtonDown(const FGeometry& MyGeometry, const FPointerEvent& MouseEvent) override;
 	virtual FReply OnMouseButtonUp(const FGeometry& MyGeometry, const FPointerEvent& MouseEvent) override;
 	virtual FReply OnMouseMove(const FGeometry& MyGeometry, const FPointerEvent& MouseEvent) override;
@@ -19,14 +21,12 @@ public:
 
 	/** Constructs this widget with InArgs */
 	void Construct(const FArguments& InArgs, TSharedRef<class SVisualLoggerView>, TSharedRef<FVisualLoggerTimeSliderController> TimeSliderController);
-	TSharedRef<SWidget> MakeTimeline(TSharedPtr<class SVisualLoggerView> VisualLoggerView, TSharedPtr<class FVisualLoggerTimeSliderController> TimeSliderController, const FVisualLogDevice::FVisualLogEntryItem& Entry);
 	TSharedRef<SWidget> GetRightClickMenuContent();
 
 	void OnTimelineSelected(TSharedPtr<class STimelinesBar> Widget);
 	void ChangeSelection(class TSharedPtr<class STimeline>, const FPointerEvent& MouseEvent);
 	virtual bool SupportsKeyboardFocus() const override { return true; }
 
-	void OnNewLogEntry(const FVisualLogDevice::FVisualLogEntryItem& Entry);
 	void OnFiltersChanged();
 	void OnSearchChanged(const FText& Filter);
 	void OnFiltersSearchChanged(const FText& Filter);
@@ -47,7 +47,7 @@ public:
 	/**
 	* @return All currently selected nodes
 	*/
-	const TArray< TSharedPtr<class STimeline> >& GetSelectedNodes() const { return SelectedNodes; }
+	const TArray< TSharedPtr<class STimeline> >& GetSelectedNodes() const { return CachedSelectedTimelines; }
 
 	const TArray< TSharedPtr<class STimeline> >& GetAllNodes() const { return TimelineItems; }
 
@@ -60,10 +60,17 @@ public:
 	bool IsNodeSelected(TSharedPtr<class STimeline> Node) const;
 
 protected:
+	void OnNewRowHandler(const FVisualLoggerDBRow& DBRow);
+	void OnNewItemHandler(const FVisualLoggerDBRow& BDRow, int32 ItemIndex);
+	void OnObjectSelectionChanged(const TArray<FName>& RowNames);
+	void OnRowChangedVisibility(const FName&);
+
+protected:
+	TSharedPtr<SVerticalBox> ContainingBorder;
 	TSharedPtr<class FVisualLoggerTimeSliderController> TimeSliderController;
 	TSharedPtr<class SVisualLoggerView> VisualLoggerView;
+
 	TArray<TSharedPtr<class STimeline> > TimelineItems;
-	TArray< TSharedPtr<class STimeline> > SelectedNodes;
-	TSharedPtr<SVerticalBox> ContainingBorder;
+	TArray< TSharedPtr<class STimeline> > CachedSelectedTimelines;
 	float CachedMinTime, CachedMaxTime;
 };
diff --git a/Engine/Source/Developer/LogVisualizer/Private/SVisualLogger.cpp b/Engine/Source/Developer/LogVisualizer/Private/SVisualLogger.cpp
--- a/Engine/Source/Developer/LogVisualizer/Private/SVisualLogger.cpp
+++ b/Engine/Source/Developer/LogVisualizer/Private/SVisualLogger.cpp
@@ -23,6 +23,8 @@
 
 #define LOCTEXT_NAMESPACE "SVisualLogger"
 
+DEFINE_LOG_CATEGORY_STATIC(LogVisualLogger, Log, All);
+
 /* Local constants
 *****************************************************************************/
 static const FName ToolbarTabId("Toolbar");
@@ -38,81 +40,116 @@ namespace LogVisualizer
 	static const FString SaveFileTypes = FString::Printf(TEXT("%s (*.%s)|*.%s"), *LogFileDescription, VISLOG_FILENAME_EXT, VISLOG_FILENAME_EXT);
 }
 
-SVisualLogger::FVisualLoggerDevice::FVisualLoggerDevice(SVisualLogger* InOwner)
-	:Owner(InOwner)
-{
-}
+DECLARE_DELEGATE_TwoParams(FOnWorldChanged, UWorld*, UWorld*);
 
-SVisualLogger::FVisualLoggerDevice::~FVisualLoggerDevice() 
-{
-}
-
-void SVisualLogger::FVisualLoggerDevice::Serialize(const class UObject* LogOwner, FName OwnerName, FName OwnerClassName, const FVisualLogEntry& LogEntry)
+/* SMessagingDebugger constructors
+*****************************************************************************/
+namespace
 {
-	if (Owner == NULL)
+	static UWorld* GetWorldForGivenObject(const class UObject* Object)
 	{
-		return;
+		UWorld* World = Object ? GEngine->GetWorldFromContextObject(Object, false) : nullptr;
+#if WITH_EDITOR
+		UEditorEngine *EEngine = Cast<UEditorEngine>(GEngine);
+		if (GIsEditor && EEngine != nullptr && World == nullptr)
+		{
+			// lets use PlayWorld during PIE/Simulate and regular world from editor otherwise, to draw debug information
+			World = EEngine->PlayWorld != nullptr ? EEngine->PlayWorld : EEngine->GetEditorWorldContext().World();
+		}
+
+#endif
+		if (!GIsEditor && World == nullptr)
+		{
+			World = GEngine->GetWorld();
+		}
+
+		return World;
 	}
 
-	UWorld* World = FLogVisualizer::Get().GetWorld();
-	if (LastWorld.Get() != World)
+	static AVisualLoggerRenderingActor* GetVisualLoggerHelperActor(UWorld* World)
 	{
-		Owner->OnNewWorld(LastWorld.Get());
-		LastWorld = World;
+		if (World == nullptr)
+		{
+			return nullptr;
+		}
+
+		for (TActorIterator<AVisualLoggerRenderingActor> It(World); It; ++It)
+		{
+			return *It;
+		}
+
+		FActorSpawnParameters SpawnInfo;
+		SpawnInfo.SpawnCollisionHandlingOverride = ESpawnActorCollisionHandlingMethod::AlwaysSpawn;
+		SpawnInfo.Name = *FString::Printf(TEXT("VisualLoggerRenderingActor"));
+		return World->SpawnActor<AVisualLoggerRenderingActor>(SpawnInfo);
 	}
 
-	Owner->OnNewLogEntry(FVisualLogDevice::FVisualLogEntryItem(OwnerName, OwnerClassName, LogEntry));
 }
 
-/* SMessagingDebugger constructors
-*****************************************************************************/
-
-SVisualLogger::SVisualLogger() : SCompoundWidget(), CommandList(MakeShareable(new FUICommandList))
+SVisualLogger::SVisualLogger() 
+	: SCompoundWidget(), CommandList(MakeShareable(new FUICommandList))
 { 
 	bPausedLogger = false;
 	bGotHistogramData = false;
-	InternalDevice = MakeShareable(new FVisualLoggerDevice(this));
+
+	class FVisualLoggerDevice : public FVisualLogDevice
+	{
+	public:
+		FVisualLoggerDevice(SVisualLogger* InVisualLogger, FOnWorldChanged OnWorldChangedDelegate) : VisualLoggerWidget(InVisualLogger), LastUsedWorld(nullptr), OnWorldChanged(OnWorldChangedDelegate) {}
+
+		virtual void Serialize(const class UObject* LogOwner, FName OwnerName, FName OwnerClassName, const FVisualLogEntry& LogEntry) override
+		{
+			VisualLoggerWidget->OnNewLogEntry(FVisualLogDevice::FVisualLogEntryItem(OwnerName, OwnerClassName, LogEntry));
+			UWorld* CurrentWorld = GetWorldForGivenObject(LogOwner);
+			if (LastUsedWorld != CurrentWorld)
+			{
+				OnWorldChanged.ExecuteIfBound(LastUsedWorld, CurrentWorld);
+				LastUsedWorld = CurrentWorld;
+			}
+		}
+
+		UWorld* LastUsedWorld;
+		FOnWorldChanged OnWorldChanged;
+		SVisualLogger* VisualLoggerWidget;
+	};
+
+	InternalDevice = MakeShareable(new FVisualLoggerDevice(this, FOnWorldChanged::CreateLambda([this](UWorld* PreviousWorld, UWorld* CurrentWorld){ OnNewWorld(CurrentWorld); })));
 	FVisualLogger::Get().AddDevice(InternalDevice.Get());
 }
 
 SVisualLogger::~SVisualLogger()
 {
-	UWorld* World = NULL;
-#if WITH_EDITOR
-	FCategoryFiltersManager::Get().SavePresistentData();
+	GEngine->OnWorldAdded().RemoveAll(this);
 
-	UEditorEngine *EEngine = Cast<UEditorEngine>(GEngine);
-	if (GIsEditor && EEngine != NULL)
-	{
-		// lets use PlayWorld during PIE/Simulate and regular world from editor otherwise, to draw debug information
-		World = EEngine->PlayWorld != NULL ? EEngine->PlayWorld : EEngine->GetEditorWorldContext().World();
+	FVisualLogger::Get().RemoveDevice(InternalDevice.Get());
+	InternalDevice.Reset();
 
-	}
-	else
+	UWorld* World = NULL;
+#if WITH_EDITOR
+	ULogVisualizerSettings::StaticClass()->GetDefaultObject<ULogVisualizerSettings>()->SavePresistentData();
 #endif
-	if (!GIsEditor)
-	{
-	World = GEngine->GetWorld();
-	}
-
-	if (World == NULL)
-	{
-		World = GWorld;
-	}
 
-	if (World)
+	if (LastUsedWorld.IsValid())
 	{
-		for (TActorIterator<AVisualLoggerRenderingActor> It(World); It; ++It)
+		for (TActorIterator<AVisualLoggerRenderingActor> It(LastUsedWorld.Get()); It; ++It)
 		{
-			World->DestroyActor(*It);
+			LastUsedWorld->DestroyActor(*It);
 		}
 	}
 
+
 	UDebugDrawService::Unregister(DrawOnCanvasDelegateHandle);
 	VisualLoggerCanvasRenderer.Reset();
 
-	FVisualLogger::Get().RemoveDevice(InternalDevice.Get());
-	InternalDevice.Reset();
+	FVisualLoggerDatabase::Get().GetEvents().OnRowSelectionChanged.RemoveAll(this);
+	FVisualLoggerDatabase::Get().GetEvents().OnNewItem.RemoveAll(this);
+	FVisualLoggerDatabase::Get().GetEvents().OnItemSelectionChanged.RemoveAll(this);
+
+	FLogVisualizer::Get().GetEvents().OnFiltersChanged.RemoveAll(this);
+	FLogVisualizer::Get().GetEvents().OnLogLineSelectionChanged.Unbind();
+	FLogVisualizer::Get().GetEvents().OnKeyboardEvent.Unbind();
+
+	FVisualLoggerDatabase::Get().Reset();
 }
 
 /* SMessagingDebugger interface
@@ -126,18 +163,24 @@ void SVisualLogger::Construct(const FArguments& InArgs, const TSharedRef<SDockTa
 	FLogVisualizer::Get().SetCurrentVisualizer(SharedThis(this));
 	//////////////////////////////////////////////////////////////////////////
 	// Visual Logger Events
-	FLogVisualizer::Get().GetVisualLoggerEvents().OnItemSelectionChanged = FOnItemSelectionChanged::CreateRaw(this, &SVisualLogger::OnItemSelectionChanged);
-	FLogVisualizer::Get().GetVisualLoggerEvents().OnFiltersChanged = FOnFiltersChanged::CreateRaw(this, &SVisualLogger::OnFiltersChanged);
-	FLogVisualizer::Get().GetVisualLoggerEvents().OnObjectSelectionChanged = FOnObjectSelectionChanged::CreateRaw(this, &SVisualLogger::OnObjectSelectionChanged);
-	FLogVisualizer::Get().GetVisualLoggerEvents().OnLogLineSelectionChanged = FOnLogLineSelectionChanged::CreateRaw(this, &SVisualLogger::OnLogLineSelectionChanged);
-	FLogVisualizer::Get().GetVisualLoggerEvents().OnKeyboardEvent = FOnKeyboardEvent::CreateRaw(this, &SVisualLogger::OnKeyboaedRedirection);
+	FLogVisualizer::Get().GetEvents().OnFiltersChanged.AddRaw(this, &SVisualLogger::OnFiltersChanged);
+	FLogVisualizer::Get().GetEvents().OnLogLineSelectionChanged = FOnLogLineSelectionChanged::CreateRaw(this, &SVisualLogger::OnLogLineSelectionChanged);
+	FLogVisualizer::Get().GetEvents().OnKeyboardEvent = FOnKeyboardEvent::CreateRaw(this, &SVisualLogger::OnKeyboaedRedirection);
+	FLogVisualizer::Get().GetTimeSliderController().Get()->GetTimeSliderArgs().OnScrubPositionChanged = FOnScrubPositionChanged::CreateRaw(this, &SVisualLogger::OnScrubPositionChanged);
+
+	FVisualLoggerDatabase::Get().GetEvents().OnRowSelectionChanged.AddRaw(this, &SVisualLogger::OnObjectSelectionChanged);
+	FVisualLoggerDatabase::Get().GetEvents().OnNewItem.AddRaw(this, &SVisualLogger::OnNewItemHandler);
+	FVisualLoggerDatabase::Get().GetEvents().OnItemSelectionChanged.AddRaw(this, &SVisualLogger::OnItemsSelectionChanged);
+
+	GEngine->OnWorldAdded().AddRaw(this, &SVisualLogger::OnNewWorld);
+
 	//////////////////////////////////////////////////////////////////////////
 	// Command Action Lists
 	const FVisualLoggerCommands& Commands = FVisualLoggerCommands::Get();
 	FUICommandList& ActionList = *CommandList;
 
 	ULogVisualizerSettings* Settings = ULogVisualizerSettings::StaticClass()->GetDefaultObject<ULogVisualizerSettings>();
-	FCategoryFiltersManager::Get().LoadPresistentData();
+	Settings->LoadPresistentData();
 
 	ActionList.MapAction(Commands.StartRecording, FExecuteAction::CreateRaw(this, &SVisualLogger::HandleStartRecordingCommandExecute), FCanExecuteAction::CreateRaw(this, &SVisualLogger::HandleStartRecordingCommandCanExecute), FIsActionChecked(), FIsActionButtonVisible::CreateRaw(this, &SVisualLogger::HandleStartRecordingCommandIsVisible));
 	ActionList.MapAction(Commands.StopRecording, FExecuteAction::CreateRaw(this, &SVisualLogger::HandleStopRecordingCommandExecute), FCanExecuteAction::CreateRaw(this, &SVisualLogger::HandleStopRecordingCommandCanExecute), FIsActionChecked(), FIsActionButtonVisible::CreateRaw(this, &SVisualLogger::HandleStopRecordingCommandIsVisible));
@@ -153,22 +196,15 @@ void SVisualLogger::Construct(const FArguments& InArgs, const TSharedRef<SDockTa
 		FIsActionButtonVisible::CreateRaw(this, &SVisualLogger::HandleCameraCommandCanExecute));
 	ActionList.MapAction(Commands.ToggleGraphs,
 		FExecuteAction::CreateLambda([](){bool& bEnableGraphsVisualization = ULogVisualizerSessionSettings::StaticClass()->GetDefaultObject<ULogVisualizerSessionSettings>()->bEnableGraphsVisualization; bEnableGraphsVisualization = !bEnableGraphsVisualization; }),
-		FCanExecuteAction::CreateLambda([this]()->bool{return bGotHistogramData; }),
+		FCanExecuteAction::CreateLambda([this]()->bool{return FVisualLoggerGraphsDatabase::Get().ContainsHistogramGraphs(); }),
 		FIsActionChecked::CreateLambda([]()->bool{return ULogVisualizerSessionSettings::StaticClass()->GetDefaultObject<ULogVisualizerSessionSettings>()->bEnableGraphsVisualization; }),
-		FIsActionButtonVisible::CreateLambda([this]()->bool{return bGotHistogramData; }));
+		FIsActionButtonVisible::CreateLambda([this]()->bool{return FVisualLoggerGraphsDatabase::Get().ContainsHistogramGraphs(); }));
 	ActionList.MapAction(Commands.ResetData, 
 		FExecuteAction::CreateRaw(this, &SVisualLogger::ResetData),
 		FCanExecuteAction::CreateRaw(this, &SVisualLogger::HandleSaveCommandCanExecute),
 		FIsActionChecked(), 
 		FIsActionButtonVisible::CreateRaw(this, &SVisualLogger::HandleSaveCommandCanExecute));
 
-	ActionList.MapAction(
-		Commands.MoveCursorLeft,
-		FExecuteAction::CreateRaw(this, &SVisualLogger::OnMoveCursorLeftCommand));
-	ActionList.MapAction(
-		Commands.MoveCursorRight,
-		FExecuteAction::CreateRaw(this, &SVisualLogger::OnMoveCursorRightCommand));
-
 
 	// Tab Spawners
 	TabManager = FGlobalTabmanager::Get()->NewTabManager(ConstructUnderMajorTab);
@@ -246,21 +282,8 @@ void SVisualLogger::Construct(const FArguments& InArgs, const TSharedRef<SDockTa
 		);
 	TabManager->SetOnPersistLayout(FTabManager::FOnPersistLayout::CreateRaw(this, &SVisualLogger::HandleTabManagerPersistLayout));
 
-	// Load the potentially previously saved new layout
-	TSharedRef<FTabManager::FLayout> UserConfiguredNewLayout = FLayoutSaveRestore::LoadFromConfig(GEditorLayoutIni, Layout);
-
-	TabManager->RestoreFrom(UserConfiguredNewLayout, TSharedPtr<SWindow>());
-
 	// Window Menu
 	FMenuBarBuilder MenuBarBuilder = FMenuBarBuilder(TSharedPtr<FUICommandList>());
-#if 0 //disabled File menu for now (SebaK)
-	MenuBarBuilder.AddPullDownMenu(
-		LOCTEXT("FileMenuLabel", "File"),
-		FText::GetEmpty(),
-		FNewMenuDelegate::CreateRaw(this, &SVisualLogger::FillFileMenu, TabManager),
-		"File"
-		);
-#endif
 	MenuBarBuilder.AddPullDownMenu(
 		LOCTEXT("WindowMenuLabel", "Window"),
 		FText::GetEmpty(),
@@ -304,67 +327,34 @@ void SVisualLogger::Construct(const FArguments& InArgs, const TSharedRef<SDockTa
 	VisualLoggerCanvasRenderer = MakeShareable(new FVisualLoggerCanvasRenderer());
 
 	DrawOnCanvasDelegateHandle = UDebugDrawService::Register(TEXT("VisLog"), FDebugDrawDelegate::CreateRaw(VisualLoggerCanvasRenderer.Get(), &FVisualLoggerCanvasRenderer::DrawOnCanvas));
-	//UGameplayDebuggingComponent::OnDebuggingTargetChangedDelegate.AddSP(this, &SVisualLogger::SelectionChanged);
-}
 
-void SVisualLogger::HandleMajorTabPersistVisualState()
-{
-	// save any settings here
+	AVisualLoggerRenderingActor* HelperActor = GetVisualLoggerHelperActor(LastUsedWorld.Get());
 }
 
-void SVisualLogger::HandleTabManagerPersistLayout(const TSharedRef<FTabManager::FLayout>& LayoutToSave)
-{
-	FLayoutSaveRestore::SaveToConfig(GEditorLayoutIni, LayoutToSave);
-}
-
-void SVisualLogger::FillFileMenu(FMenuBuilder& MenuBuilder, const TSharedPtr<FTabManager> InTabManager)
+void SVisualLogger::OnNewLogEntry(const FVisualLogDevice::FVisualLogEntryItem& Entry)
 {
-	MenuBuilder.BeginSection("LogFile", LOCTEXT("FileMenu", "Log File"));
+	if (bPausedLogger)
 	{
-		MenuBuilder.AddMenuEntry(FVisualLoggerCommands::Get().LoadFromVLog);
-		MenuBuilder.AddMenuEntry(FVisualLoggerCommands::Get().SaveToVLog);
-		MenuBuilder.AddMenuEntry(FVisualLoggerCommands::Get().SaveAllToVLog);
+		OnPauseCacheForEntries.Add(Entry);
+		return;
 	}
-	MenuBuilder.EndSection();
-	MenuBuilder.BeginSection("LogFilters", LOCTEXT("FIlterMenu", "Log Filters"));
-	{
-		MenuBuilder.AddSubMenu(
-			LOCTEXT("LoadPreset", "Load Preset"), LOCTEXT("LoadPresetTooltip", "Load filter's preset"),
-			FNewMenuDelegate::CreateRaw(this, &SVisualLogger::FillLoadPresetMenu));
 
-		MenuBuilder.AddMenuEntry(
-			LOCTEXT("SavePreset", "Save Preset"), LOCTEXT("SavePresetTooltip", "Save filter's setting as preset"),
-			FSlateIcon(), FUIAction(
-			FExecuteAction::CreateLambda(
-				[this](){
-				}
-			)));
+	FVisualLoggerDatabase::Get().AddItem(Entry);
+
+	if (ULogVisualizerSettings::StaticClass()->GetDefaultObject<ULogVisualizerSettings>()->bStickToRecentData)
+	{
+		FLogVisualizer::Get().GetTimeSliderController()->CommitScrubPosition(Entry.Entry.TimeStamp, false);
 	}
-	MenuBuilder.EndSection();
 }
 
-void SVisualLogger::FillLoadPresetMenu(FMenuBuilder& MenuBuilder)
+void SVisualLogger::HandleMajorTabPersistVisualState()
 {
-	ULogVisualizerSettings* Settings = ULogVisualizerSettings::StaticClass()->GetDefaultObject<ULogVisualizerSettings>();
-	
-	MenuBuilder.BeginSection("FilterPresets", LOCTEXT("FilterPresets", "Presets"));
-#if 0 //disabled for now, we don't use presets yet
-	for (auto& CurrentPreset : Settings->FilterPresets)
-	{
-		MenuBuilder.AddMenuEntry(
-			FText::FromString(CurrentPreset.FilterName), LOCTEXT("LoadPresetTooltip", "Load filter's preset"),
-			FSlateIcon(), FUIAction(
-			FExecuteAction::CreateLambda([this, CurrentPreset](){ this->SetFiltersPreset(CurrentPreset); }
-			))
-		);
-	}
-#endif
-	MenuBuilder.EndSection();
+	// save any settings here
 }
 
-void SVisualLogger::SetFiltersPreset(const struct FFiltersPreset& Preset)
+void SVisualLogger::HandleTabManagerPersistLayout(const TSharedRef<FTabManager::FLayout>& LayoutToSave)
 {
-
+	// save any layout here
 }
 
 void SVisualLogger::FillWindowMenu(FMenuBuilder& MenuBuilder, const TSharedPtr<FTabManager> TabManager)
@@ -408,6 +398,7 @@ TSharedRef<SDockTab> SVisualLogger::HandleTabManagerSpawnTab(const FSpawnTabArgs
 		AutoSizeTab = false;
 	}
 
+	check(TabWidget.IsValid());
 	return SNew(SVisualLoggerTab)
 		.ShouldAutosize(AutoSizeTab)
 		.TabRole(ETabRole::DocumentTab)
@@ -443,7 +434,7 @@ void SVisualLogger::HandleStopRecordingCommandExecute()
 {
 	UWorld* World = FLogVisualizer::Get().GetWorld();
 
-	if (FParse::Param(FCommandLine::Get(), TEXT("LogNavOctree")) == false && ULogVisualizerSettings::StaticClass()->GetDefaultObject<ULogVisualizerSettings>()->bLogNavOctreeOnStop)
+	if (FParse::Param(FCommandLine::Get(), TEXT("LogNavOctree")) == true && ULogVisualizerSettings::StaticClass()->GetDefaultObject<ULogVisualizerSettings>()->bLogNavOctreeOnStop)
 	{
 		FVisualLogger::NavigationDataDump(World, LogNavigation, ELogVerbosity::Log, INDEX_NONE, FBox());
 	}
@@ -454,6 +445,11 @@ void SVisualLogger::HandleStopRecordingCommandExecute()
 	{
 		AVisualLoggerCameraController::DisableCamera(World);
 	}
+
+	if (bPausedLogger)
+	{
+		HandleResumeCommandExecute();
+	}
 }
 
 
@@ -546,11 +542,6 @@ void SVisualLogger::HandleCameraCommandExecute()
 	{
 		// switch debug cam on
 		CameraController = AVisualLoggerCameraController::EnableCamera(World);
-		if (CameraController.IsValid())
-		{
-			//CameraController->OnActorSelected = ALogVisualizerCameraController::FActorSelectedDelegate::CreateSP(this, &SVisualLogger::CameraActorSelected);
-			//CameraController->OnIterateLogEntries = ALogVisualizerCameraController::FLogEntryIterationDelegate::CreateSP(this, &SVisualLogger::IncrementCurrentLogIndex);
-		}
 	}
 }
 
@@ -593,7 +584,7 @@ void SVisualLogger::HandleLoadCommandExecute()
 
 	if (bOpened && OpenFilenames.Num() > 0)
 	{
-		OnNewWorld(nullptr);
+		OnNewWorld(GetWorldForGivenObject(nullptr));
 		for (int FilenameIndex = 0; FilenameIndex < OpenFilenames.Num(); ++FilenameIndex)
 		{
 			FString CurrentFileName = OpenFilenames[FilenameIndex];
@@ -617,9 +608,7 @@ void SVisualLogger::HandleLoadCommandExecute()
 
 bool SVisualLogger::HandleSaveCommandCanExecute() const
 { 
-	TArray<TSharedPtr<class STimeline> > OutTimelines;
-	MainView->GetTimelines(OutTimelines);
-	return OutTimelines.Num() > 0;
+	return FVisualLoggerDatabase::Get().NumberOfRows() > 0;
 }
 
 void SVisualLogger::HandleSaveAllCommandExecute()
@@ -634,14 +623,20 @@ void SVisualLogger::HandleSaveCommandExecute()
 
 void SVisualLogger::HandleSaveCommand(bool bSaveAllData)
 {
-	TArray<TSharedPtr<class STimeline> > OutTimelines;
-	MainView->GetTimelines(OutTimelines, !bSaveAllData);
-	if (OutTimelines.Num() == 0)
+	TArray<FName> SelectedRows;
+	if (!bSaveAllData)
 	{
-		MainView->GetTimelines(OutTimelines);
+		SelectedRows = FVisualLoggerDatabase::Get().GetSelectedRows();
+	}
+	else
+	{
+		for (auto Iter(FVisualLoggerDatabase::Get().GetConstRowIterator()); Iter; ++Iter)
+		{
+			SelectedRows.Add((*Iter).GetOwnerName());
+		}
 	}
 
-	if (OutTimelines.Num())
+	if (SelectedRows.Num())
 	{
 		// Prompt the user for the filenames
 		TArray<FString> SaveFilenames;
@@ -675,9 +670,10 @@ void SVisualLogger::HandleSaveCommand(bool bSaveAllData)
 			if (SaveFilenames.Num() > 0)
 			{
 				TArray<FVisualLogDevice::FVisualLogEntryItem> FrameCache;
-				for (auto CurrentItem : OutTimelines)
+				for (auto CurrentName : SelectedRows)
 				{
-					FrameCache.Append(CurrentItem->GetEntries());
+					FVisualLoggerDBRow& DataRow = FVisualLoggerDatabase::Get().GetRowByName(CurrentName);
+					FrameCache.Append(DataRow.GetItems());
 				}
 
 				if (FrameCache.Num())
@@ -695,6 +691,13 @@ void SVisualLogger::HandleSaveCommand(bool bSaveAllData)
 
 void SVisualLogger::ResetData()
 {
+	bGotHistogramData = false;
+	OnPauseCacheForEntries.Reset();
+
+	FLogVisualizer::Get().Reset();
+	FVisualLoggerDatabase::Get().Reset();
+	FVisualLoggerFilters::Get().Reset();
+
 	if (MainView.IsValid())
 	{
 		MainView->ResetData();
@@ -707,134 +710,278 @@ void SVisualLogger::ResetData()
 
 	if (LogsList.IsValid())
 	{
-		LogsList->OnItemSelectionChanged(FVisualLogDevice::FVisualLogEntryItem());
+		LogsList->ResetData();
 	}
 
 	if (StatusView.IsValid())
 	{
-		StatusView->OnItemSelectionChanged(FVisualLogDevice::FVisualLogEntryItem());
+		StatusView->ResetData();
 	}
 
 	if (VisualLoggerCanvasRenderer.IsValid())
 	{
-		VisualLoggerCanvasRenderer->OnItemSelectionChanged(FVisualLogEntry());
-		VisualLoggerCanvasRenderer->ObjectSelectionChanged(NULL);
+		VisualLoggerCanvasRenderer->ResetData();
 	}
 
-	AVisualLoggerRenderingActor* HelperActor = Cast<AVisualLoggerRenderingActor>(FLogVisualizer::Get().GetVisualLoggerHelperActor());
-	if (HelperActor)
+	if (AVisualLoggerRenderingActor* HelperActor = GetVisualLoggerHelperActor(LastUsedWorld.Get()))
 	{
-		HelperActor->OnItemSelectionChanged(FVisualLogDevice::FVisualLogEntryItem());
-		TArray<TSharedPtr<class STimeline> > TimeLines;
-		HelperActor->ObjectSelectionChanged(TimeLines);
+		HelperActor->ResetRendering();
 	}
 
-	bGotHistogramData = false;
-	OnPauseCacheForEntries.Reset();
+	const TMap<FName, FVisualLogExtensionInterface*>& AllExtensions = FVisualLogger::Get().GetAllExtensions();
+	for (auto Iterator = AllExtensions.CreateConstIterator(); Iterator; ++Iterator)
+	{
+		FVisualLogExtensionInterface* Extension = (*Iterator).Value;
+		if (Extension != NULL)
+		{
+			Extension->ResetData(FVisualLoggerEditorInterface::Get());
+		}
+	}
+
+	FLogVisualizer::Get().GetEvents().OnLogLineSelectionChanged = FOnLogLineSelectionChanged::CreateRaw(this, &SVisualLogger::OnLogLineSelectionChanged);
+	FLogVisualizer::Get().GetEvents().OnKeyboardEvent = FOnKeyboardEvent::CreateRaw(this, &SVisualLogger::OnKeyboaedRedirection);
+	FLogVisualizer::Get().GetTimeSliderController().Get()->GetTimeSliderArgs().OnScrubPositionChanged = FOnScrubPositionChanged::CreateRaw(this, &SVisualLogger::OnScrubPositionChanged);
 }
 
 void SVisualLogger::OnNewWorld(UWorld* NewWorld)
 {
-	InternalDevice->SerLastWorld(NewWorld);
-	if (ULogVisualizerSettings::StaticClass()->GetDefaultObject<ULogVisualizerSettings>()->bResetDataWithNewSession)
+	if (LastUsedWorld.IsValid() && LastUsedWorld != NewWorld)
+	{
+		for (TActorIterator<AVisualLoggerRenderingActor> It(LastUsedWorld.Get()); It; ++It)
+		{
+			LastUsedWorld->DestroyActor(*It);
+		}
+	}
+	LastUsedWorld = NewWorld;
+
+	AVisualLoggerRenderingActor* HelperActor = GetVisualLoggerHelperActor(LastUsedWorld.Get());
+	check(HelperActor);
+
+	if (LastUsedWorld.IsValid() == false || ULogVisualizerSettings::StaticClass()->GetDefaultObject<ULogVisualizerSettings>()->bResetDataWithNewSession)
 	{
 		ResetData();
 	}
+
+	// reset data and simulate row/item selection to recreate rendering proxy with correct data
+	HelperActor->ResetRendering();
+	const TArray<FName>& SelectedRows = FVisualLoggerDatabase::Get().GetSelectedRows();
+	HelperActor->ObjectSelectionChanged(SelectedRows);
+	for (auto& RowName : SelectedRows)
+	{
+		FVisualLoggerDBRow& DBRow = FVisualLoggerDatabase::Get().GetRowByName(RowName);
+		HelperActor->OnItemSelectionChanged(DBRow, DBRow.GetCurrentItemIndex());
+	}
+
 }
 
-void SVisualLogger::OnNewLogEntry(const FVisualLogDevice::FVisualLogEntryItem& Entry)
+void SVisualLogger::OnObjectSelectionChanged(const TArray<FName>& RowNames)
 {
-	if (!bPausedLogger)
+	const float ScrubTime = FLogVisualizer::Get().GetTimeSliderController().Get()->GetTimeSliderArgs().ScrubPosition.Get();
+	for (auto RowName : RowNames)
 	{
-		CollectNewCategories(Entry);
-		MainView->OnNewLogEntry(Entry);
-		bGotHistogramData = bGotHistogramData || Entry.Entry.HistogramSamples.Num() > 0;
+		FVisualLoggerDBRow& DBRow = FVisualLoggerDatabase::Get().GetRowByName(RowName);
+		if (DBRow.GetCurrentItemIndex() == INDEX_NONE)
+		{
+			DBRow.MoveTo(DBRow.GetClosestItem(ScrubTime, ScrubTime));
+		}
 	}
-	else
+}
+
+void SVisualLogger::OnItemsSelectionChanged(const FVisualLoggerDBRow& ChangedRow, int32 SelectedItemIndex)
+{
+	const TMap<FName, FVisualLogExtensionInterface*>& AllExtensions = FVisualLogger::Get().GetAllExtensions();
+	for (auto Iterator = AllExtensions.CreateConstIterator(); Iterator; ++Iterator)
 	{
-		OnPauseCacheForEntries.Add(Entry);
+		FVisualLogExtensionInterface* Extension = (*Iterator).Value;
+		if (Extension != NULL)
+		{
+			Extension->OnItemsSelectionChanged(FVisualLoggerEditorInterface::Get());
+		}
 	}
 }
 
-void SVisualLogger::CollectNewCategories(const FVisualLogDevice::FVisualLogEntryItem& Item)
+void SVisualLogger::OnFiltersChanged()
 {
-	TArray<FVisualLoggerCategoryVerbosityPair> Categories;
-	FVisualLoggerHelpers::GetCategories(Item.Entry, Categories);
-	for (auto& CategoryAndVerbosity : Categories)
+	const uint32 StartCycles = FPlatformTime::Cycles();
+
+	ULogVisualizerSettings* Settings = ULogVisualizerSettings::StaticClass()->GetDefaultObject<ULogVisualizerSettings>();
+	const FString QuickSearchStrng = FVisualLoggerFilters::Get().GetSearchString();
+
+	TArray<TFuture<void> > AllFutures;
+	for (auto Iterator = FVisualLoggerDatabase::Get().GetRowIterator(); Iterator; ++Iterator)
 	{
-		VisualLoggerFilters->AddFilter(CategoryAndVerbosity.CategoryName.ToString());
+		FVisualLoggerDBRow* DBRow = &(*Iterator);
+		AllFutures.Add(
+			Async<void>(EAsyncExecution::TaskGraph, [this, DBRow]()
+			{
+				const TArray<FVisualLogDevice::FVisualLogEntryItem>& Entries = DBRow->GetItems();
+				for (int32 Index = 0; Index < Entries.Num(); ++Index)
+				{
+					UpdateVisibilityForEntry(*DBRow, Index);
+				}
+			}
+		));
 	}
 
-	TMap<FString, TArray<FString> > OutCategories;
-	FVisualLoggerHelpers::GetHistogramCategories(Item.Entry, OutCategories);
-	for (const auto& CurrentCategory : OutCategories)
+	bool bAllFuturesReady = false;
+	do
 	{
-		for (const auto& CurrentDataName : CurrentCategory.Value)
+		bAllFuturesReady = true;
+		for (TFuture<void>& CurrentFuture : AllFutures)
+		{
+			bAllFuturesReady &= CurrentFuture.IsReady();
+			if (bAllFuturesReady == false)
+			{
+				break;
+			}
+		}
+		if (bAllFuturesReady == false)
 		{
-			VisualLoggerFilters->AddFilter(CurrentCategory.Key, CurrentDataName);
+			FPlatformProcess::Sleep(0.01);
 		}
+	} while (bAllFuturesReady != true);
+
+	for (auto Iterator = FVisualLoggerDatabase::Get().GetRowIterator(); Iterator; ++Iterator)
+	{
+		FVisualLoggerDBRow& DBRow = *Iterator;
+		FVisualLoggerDatabase::Get().SetRowVisibility(DBRow.GetOwnerName(), DBRow.GetNumberOfHiddenItems() != DBRow.GetItems().Num());
 	}
-}
 
-void SVisualLogger::OnItemSelectionChanged(const FVisualLogDevice::FVisualLogEntryItem& EntryItem)
-{
-	LogsList->OnItemSelectionChanged(EntryItem);
-	StatusView->OnItemSelectionChanged(EntryItem);
-	VisualLoggerCanvasRenderer->OnItemSelectionChanged(EntryItem.Entry);
-	VisualLoggerFilters->OnItemSelectionChanged(EntryItem.Entry);
-	AVisualLoggerRenderingActor* HelperActor = Cast<AVisualLoggerRenderingActor>(FLogVisualizer::Get().GetVisualLoggerHelperActor());
-	if (HelperActor)
+
+	const uint32 EndCycles = FPlatformTime::Cycles();
+	const int32 BlockingCycles = int32(FPlatformTime::Cycles() - StartCycles);
 	{
-		HelperActor->OnItemSelectionChanged(EntryItem);
+		const TArray<FName>& SelectedRows = FVisualLoggerDatabase::Get().GetSelectedRows();
+		const float ScrubTime = FLogVisualizer::Get().GetTimeSliderController()->GetTimeSliderArgs().ScrubPosition.Get();
+		{
+			for (auto RowName : SelectedRows)
+			{
+				auto& DBRow = FVisualLoggerDatabase::Get().GetRowByName(RowName);
+				const int32 ClosestItem = DBRow.GetClosestItem(ScrubTime, ScrubTime);
+				const TArray<FVisualLogDevice::FVisualLogEntryItem>& Items = DBRow.GetItems();
+				if (Items.IsValidIndex(ClosestItem) && Items[ClosestItem].Entry.TimeStamp <= ScrubTime)
+				{
+					DBRow.MoveTo(ClosestItem);
+				}
+			}
+		}
 	}
+	UE_LOG(LogVisualLogger, Display, TEXT("SVisualLogger::OnFiltersChanged: %5.2fms"), FPlatformTime::ToMilliseconds(BlockingCycles));
 }
 
-void SVisualLogger::OnFiltersChanged()
+void SVisualLogger::OnFiltersSearchChanged(const FText& Filter)
 {
-	LogsList->OnFiltersChanged();
-	MainView->OnFiltersChanged();
-	if (AVisualLoggerRenderingActor* HelperActor = Cast<AVisualLoggerRenderingActor>(FLogVisualizer::Get().GetVisualLoggerHelperActor()))
+	const uint32 StartCycles = FPlatformTime::Cycles();
+
+	FVisualLoggerFilters::Get().SetSearchString(Filter.ToString());
+
+	ULogVisualizerSettings* Settings = ULogVisualizerSettings::StaticClass()->GetDefaultObject<ULogVisualizerSettings>();
+	const FString QuickSearchStrng = FVisualLoggerFilters::Get().GetSearchString();
+
+	TArray<TFuture<void> > AllFutures;
+	for (auto Iterator = FVisualLoggerDatabase::Get().GetRowIterator(); Iterator; ++Iterator)
 	{
-		HelperActor->OnItemSelectionChanged(LogsList->GetCurrentLogEntry());
-		HelperActor->MarkComponentsRenderStateDirty();
+		FVisualLoggerDBRow* DBRow = &(*Iterator);
+			AllFutures.Add(
+				Async<void>(EAsyncExecution::TaskGraph, [this, DBRow]()
+					{
+						const TArray<FVisualLogDevice::FVisualLogEntryItem>& Entries = DBRow->GetItems();
+						for (int32 Index = 0; Index < Entries.Num(); ++Index)
+						{
+							UpdateVisibilityForEntry(*DBRow, Index);
+						}
+					}
+				)
+		);
 	}
-	VisualLoggerCanvasRenderer->OnItemSelectionChanged(LogsList->GetCurrentLogEntry().Entry);
-}
 
-void SVisualLogger::OnObjectSelectionChanged(TSharedPtr<class STimeline> TimeLine)
-{
-	TArray<TSharedPtr<class STimeline> > TimeLines;
-	GetTimelines(TimeLines, true);
+	bool bAllFuturesReady = false;
+	do 
+	{
+		bAllFuturesReady = true;
+		for (TFuture<void>& CurrentFuture : AllFutures)
+		{
+			bAllFuturesReady &= CurrentFuture.IsReady();
+			if (bAllFuturesReady == false)
+			{
+				break;
+			}
+		}
+		if (bAllFuturesReady == false)
+		{
+			FPlatformProcess::Sleep(0.01);
+		}
+	} while (bAllFuturesReady != true);
 
-	LogsList->ObjectSelectionChanged(TimeLines);
+	for (auto Iterator = FVisualLoggerDatabase::Get().GetRowIterator(); Iterator; ++Iterator)
+	{
+		FVisualLoggerDBRow& DBRow = *Iterator;
+		FVisualLoggerDatabase::Get().SetRowVisibility(DBRow.GetOwnerName(), DBRow.GetNumberOfHiddenItems() != DBRow.GetItems().Num());
+	}
+
+	if (LogsList.IsValid())
+	{
+		// it depends on rows visibility so it have to be called here, manually after changes to rows visibilities
+		LogsList->OnFiltersSearchChanged(Filter);
+	}
 
-	VisualLoggerCanvasRenderer->ObjectSelectionChanged(TimeLine);
-	AVisualLoggerRenderingActor* HelperActor = Cast<AVisualLoggerRenderingActor>(FLogVisualizer::Get().GetVisualLoggerHelperActor());
-	if (HelperActor)
+	if (VisualLoggerCanvasRenderer.IsValid())
 	{
-		HelperActor->ObjectSelectionChanged(TimeLines);
+		VisualLoggerCanvasRenderer->DirtyCachedData();
 	}
-	MainView->OnObjectSelectionChanged(TimeLine);
-	FLogVisualizer::Get().OnObjectSelectionChanged(TimeLine);
 
-	StatusView->ObjectSelectionChanged(TimeLines);
+	const uint32 EndCycles = FPlatformTime::Cycles();
+	const int32 BlockingCycles = int32(FPlatformTime::Cycles() - StartCycles);
+	UE_LOG(LogVisualLogger, Display, TEXT("SVisualLogger::OnFiltersSearchChanged: %5.2fms"), FPlatformTime::ToMilliseconds(BlockingCycles));
 }
 
-void SVisualLogger::OnFiltersSearchChanged(const FText& Filter)
+void SVisualLogger::OnNewItemHandler(const FVisualLoggerDBRow& DBRow, int32 ItemIndex)
 {
-	FCategoryFiltersManager::Get().SetSearchString(Filter.ToString());
+	UpdateVisibilityForEntry(DBRow, ItemIndex);
+}
 
-	if (VisualLoggerFilters.IsValid())
+void SVisualLogger::UpdateVisibilityForEntry(const FVisualLoggerDBRow& DBRow, int32 ItemIndex)
+{
+	ULogVisualizerSettings* Settings = ULogVisualizerSettings::StaticClass()->GetDefaultObject<ULogVisualizerSettings>();
+	const FVisualLogDevice::FVisualLogEntryItem& CurrentEntry = DBRow.GetItems()[ItemIndex];
+	FString SearchString = FVisualLoggerFilters::Get().GetSearchString();
+
+	TArray<FVisualLoggerCategoryVerbosityPair> OutCategories;
+	FVisualLoggerHelpers::GetCategories(CurrentEntry.Entry, OutCategories);
+	bool bHasValidCategories = false;
+	for (FVisualLoggerCategoryVerbosityPair& Categoryair : OutCategories)
 	{
-		VisualLoggerFilters->OnFiltersSearchChanged(Filter);
+		bHasValidCategories = bHasValidCategories || FVisualLoggerFilters::Get().MatchCategoryFilters(Categoryair.CategoryName.ToString(), Categoryair.Verbosity);
 	}
-	if (LogsList.IsValid())
+
+	if (Settings->bSearchInsideLogs && bHasValidCategories && SearchString.Len() > 0)
 	{
-		LogsList->OnFiltersSearchChanged(Filter);
+		bool bMatchSearchString = false;
+		for (const FVisualLogLine& CurrentLine : CurrentEntry.Entry.LogLines)
+		{
+			if (CurrentLine.Line.Find(SearchString) != INDEX_NONE || CurrentLine.Category.ToString().Find(SearchString) != INDEX_NONE)
+			{
+				bMatchSearchString = true;
+				break;
+			}
+		}
+		if (!bMatchSearchString)
+		{
+			for (const FVisualLogEvent& CurrentEvent : CurrentEntry.Entry.Events)
+			{
+				if (CurrentEvent.Name.Find(SearchString) != INDEX_NONE)
+				{
+					bMatchSearchString = true;
+					break;
+				}
+			}
+		}
+
+		FVisualLoggerDatabase::Get().GetRowByName(DBRow.GetOwnerName()).SetItemVisibility(ItemIndex, bMatchSearchString);
 	}
-	if (MainView.IsValid())
+	else
 	{
-		MainView->OnFiltersSearchChanged(Filter);
+		FVisualLoggerDatabase::Get().GetRowByName(DBRow.GetOwnerName()).SetItemVisibility(ItemIndex, bHasValidCategories);
 	}
 }
 
@@ -846,60 +993,117 @@ void SVisualLogger::OnLogLineSelectionChanged(TSharedPtr<struct FLogEntryItem> S
 		FVisualLogExtensionInterface* Extension = (*Iterator).Value;
 		if (Extension != NULL)
 		{
-			Extension->LogEntryLineSelectionChanged(SelectedItem, UserData, TagName);
+			Extension->OnLogLineSelectionChanged(FVisualLoggerEditorInterface::Get(), SelectedItem, UserData);
 		}
 	}
-
-	AVisualLoggerRenderingActor* HelperActor = Cast<AVisualLoggerRenderingActor>(FLogVisualizer::Get().GetVisualLoggerHelperActor());
-	if (HelperActor)
-	{
-		HelperActor->OnItemSelectionChanged(LogsList->GetCurrentLogEntry());
-		HelperActor->MarkComponentsRenderStateDirty();
-	}
 }
 
-void SVisualLogger::OnMoveCursorLeftCommand()
+void SVisualLogger::OnScrubPositionChanged(float NewScrubPosition, bool bScrubbing)
 {
-	FLogVisualizer::Get().GotoNextItem(1);
-}
-
-void SVisualLogger::OnMoveCursorRightCommand()
-{
-	FLogVisualizer::Get().GotoPreviousItem(1);
+	const TArray<FName> &SelectedRows = FVisualLoggerDatabase::Get().GetSelectedRows();
+	const float ScrubTime = FLogVisualizer::Get().GetTimeSliderController()->GetTimeSliderArgs().ScrubPosition.Get();
+	for (auto RowName : SelectedRows)
+	{
+		auto& DBRow = FVisualLoggerDatabase::Get().GetRowByName(RowName);
+		const int32 ClosestItem = SelectedRows.Num() > 1 ? DBRow.GetClosestItem(NewScrubPosition, ScrubTime) : DBRow.GetClosestItem(NewScrubPosition);
+		const TArray<FVisualLogDevice::FVisualLogEntryItem>& Items = DBRow.GetItems();
+		if (Items.IsValidIndex(ClosestItem) && Items[ClosestItem].Entry.TimeStamp <= NewScrubPosition)
+		{
+			DBRow.MoveTo(ClosestItem);
+		}
+	}
 }
 
 FReply SVisualLogger::OnKeyboaedRedirection(const FGeometry& MyGeometry, const FKeyEvent& InKeyEvent)
 {
 	FReply ReturnValue = FReply::Unhandled();
 
-	if (InKeyEvent.GetKey() == EKeys::Left)
-	{
-		FLogVisualizer::Get().GotoPreviousItem(InKeyEvent.IsLeftControlDown() ? InKeyEvent.IsLeftShiftDown() ? 20 : 10 : 1);
-		ReturnValue = FReply::Handled();
-	}
-	else if (InKeyEvent.GetKey() == EKeys::Right)
-	{
-		FLogVisualizer::Get().GotoNextItem(InKeyEvent.IsLeftControlDown() ? InKeyEvent.IsLeftShiftDown() ? 20 : 10 : 1);
-		ReturnValue = FReply::Handled();
-	}
-	else if (InKeyEvent.GetKey() == EKeys::Enter)
+	const TArray<FName>& SelectedRows = FVisualLoggerDatabase::Get().GetSelectedRows();
+	if (SelectedRows.Num() == 0)
 	{
-		FLogVisualizer::Get().MoveCamera();
-		ReturnValue = FReply::Handled();
+		return ReturnValue;
 	}
-	else if (InKeyEvent.GetKey() == EKeys::Up || InKeyEvent.GetKey() == EKeys::Down)
+
+	// find time to move by
+	const FKey Key = InKeyEvent.GetKey();
+	if (Key == EKeys::Left || Key == EKeys::Right)
 	{
+		const float ScrubTime = FLogVisualizer::Get().GetTimeSliderController()->GetTimeSliderArgs().ScrubPosition.Get();
+		float NewTimeToSet = ScrubTime;
+		float BestTimeDifference = FLT_MAX;
+
+		const int32 MoveDist = InKeyEvent.IsLeftControlDown() ? InKeyEvent.IsLeftShiftDown() ? 20 : 10 : 1;
+		for (auto RowName : SelectedRows)
+		{
+			const FVisualLoggerDBRow& DBRow = FVisualLoggerDatabase::Get().GetRowByName(RowName);
+			const int32 CurrentItemIndex = DBRow.GetCurrentItemIndex();
+			if (CurrentItemIndex == INDEX_NONE)
+			{
+				continue;
+			}
+
+			if (Key == EKeys::Right)
+			{
+				float TimeDifference = DBRow.GetCurrentItem().Entry.TimeStamp - ScrubTime;
+				if (TimeDifference > 0 && FMath::Abs(TimeDifference) < FMath::Abs(BestTimeDifference))
+				{
+					BestTimeDifference = TimeDifference;
+					NewTimeToSet = DBRow.GetCurrentItem().Entry.TimeStamp;
+				}
+
+				const int32 NextItemIndex = FLogVisualizer::Get().GetNextItem(RowName, MoveDist);
+				TimeDifference = DBRow.GetItems()[NextItemIndex].Entry.TimeStamp - ScrubTime;
+				if (TimeDifference > 0 && FMath::Abs(TimeDifference) < FMath::Abs(BestTimeDifference))
+				{
+					BestTimeDifference = TimeDifference;
+					NewTimeToSet = DBRow.GetItems()[NextItemIndex].Entry.TimeStamp;
+				}
+			}
+			else if (Key == EKeys::Left)
+			{
+				float TimeDifference = DBRow.GetCurrentItem().Entry.TimeStamp - ScrubTime;
+				if (TimeDifference < 0 && FMath::Abs(TimeDifference) < FMath::Abs(BestTimeDifference))
+				{
+					BestTimeDifference = TimeDifference;
+					NewTimeToSet = DBRow.GetCurrentItem().Entry.TimeStamp;
+				}
+
+				const int32 PrevItemIndex = FLogVisualizer::Get().GetPreviousItem(RowName, MoveDist);
+				TimeDifference = DBRow.GetItems()[PrevItemIndex].Entry.TimeStamp - ScrubTime;
+				if (TimeDifference < 0 && FMath::Abs(TimeDifference) > 0 && FMath::Abs(TimeDifference) < FMath::Abs(BestTimeDifference))
+				{
+					BestTimeDifference = TimeDifference;
+					NewTimeToSet = DBRow.GetItems()[PrevItemIndex].Entry.TimeStamp;
+				}
+			}
+		}
+
+		FLogVisualizer::Get().GetTimeSliderController()->CommitScrubPosition(NewTimeToSet, false);
 		ReturnValue = FReply::Handled();
 	}
-	return ReturnValue;
-}
 
-void SVisualLogger::GetTimelines(TArray<TSharedPtr<class STimeline> >& TimeLines, bool bOnlySelectedOnes)
-{
-	if (MainView.IsValid())
+	FName OwnerName = SelectedRows[SelectedRows.Num() - 1];
+	const FVisualLoggerDBRow& DBRow = FVisualLoggerDatabase::Get().GetRowByName(OwnerName);
+	if (DBRow.GetCurrentItemIndex() != INDEX_NONE)
 	{
-		MainView->GetTimelines(TimeLines, bOnlySelectedOnes);
+		if (Key == EKeys::Home)
+		{
+			FLogVisualizer::Get().GotoFirstItem(OwnerName);
+			ReturnValue = FReply::Handled();
+		}
+		else if (Key == EKeys::End)
+		{
+			FLogVisualizer::Get().GotoLastItem(OwnerName);
+			ReturnValue = FReply::Handled();
+		}
+		else if (Key == EKeys::Enter)
+		{
+			FLogVisualizer::Get().UpdateCameraPosition(OwnerName, DBRow.GetCurrentItemIndex());
+			ReturnValue = FReply::Handled();
+		}
 	}
+
+	return ReturnValue;
 }
 
 #undef LOCTEXT_NAMESPACE
diff --git a/Engine/Source/Developer/LogVisualizer/Private/SVisualLogger.h b/Engine/Source/Developer/LogVisualizer/Private/SVisualLogger.h
--- a/Engine/Source/Developer/LogVisualizer/Private/SVisualLogger.h
+++ b/Engine/Source/Developer/LogVisualizer/Private/SVisualLogger.h
@@ -35,38 +35,21 @@ public:
 	void Construct(const FArguments& InArgs, const TSharedRef<SDockTab>& ConstructUnderMajorTab, const TSharedPtr<SWindow>& ConstructUnderWindow);
 
 public:
-	class FVisualLoggerDevice : public FVisualLogDevice
-	{
-		friend class SVisualLogger;
-	public:
-		FVisualLoggerDevice(SVisualLogger* Owner);
-		virtual ~FVisualLoggerDevice();
-		virtual void Serialize(const class UObject* LogOwner, FName OwnerName, FName OwnerClassName, const FVisualLogEntry& LogEntry) override;
-
-	protected:
-		void SerLastWorld(class UWorld* InWorld) { LastWorld = InWorld; }
-
-		SVisualLogger* Owner;
-		TWeakObjectPtr<class UWorld>	LastWorld;
-	};
-	friend class FVisualLogDevice;
-
 	TSharedRef<SDockTab> HandleTabManagerSpawnTab(const FSpawnTabArgs& Args, FName TabIdentifier) const;
-	void FillFileMenu(FMenuBuilder& MenuBuilder, const TSharedPtr<FTabManager> TabManager);
 	static void FillWindowMenu(FMenuBuilder& MenuBuilder, const TSharedPtr<FTabManager> TabManager);
-	void FillLoadPresetMenu(FMenuBuilder& Builder);
-
-	void GetTimelines(TArray<TSharedPtr<class STimeline> >&, bool bOnlySelectedOnes = false);
 
 	/** Callback for for when the owner tab's visual state is being persisted. */
 	void HandleMajorTabPersistVisualState();
 	virtual bool SupportsKeyboardFocus() const override { return true; }
 
 	void OnNewLogEntry(const FVisualLogDevice::FVisualLogEntryItem& Entry);
-	void CollectNewCategories(const FVisualLogDevice::FVisualLogEntryItem& Entry);
-	void OnItemSelectionChanged(const FVisualLogDevice::FVisualLogEntryItem& EntryItem);
 	void OnFiltersChanged();
-	void OnObjectSelectionChanged(TSharedPtr<class STimeline> TimeLine);
+	void OnObjectSelectionChanged(const TArray<FName>& RowNames);
+	void OnItemsSelectionChanged(const FVisualLoggerDBRow&, int32);
+	void OnNewItemHandler(const FVisualLoggerDBRow& BDRow, int32 ItemIndex);
+	void UpdateVisibilityForEntry(const FVisualLoggerDBRow& BDRow, int32 ItemIndex);
+	void OnScrubPositionChanged(float NewScrubPosition, bool bScrubbing);
+
 	void OnFiltersSearchChanged(const FText& Filter);
 	void OnLogLineSelectionChanged(TSharedPtr<struct FLogEntryItem> SelectedItem, int64 UserData, FName TagName);
 	FReply OnKeyboaedRedirection(const FGeometry& MyGeometry, const FKeyEvent& InKeyEvent);
@@ -105,14 +88,9 @@ public:
 	TSharedPtr<SVisualLoggerFilters> GetVisualLoggerFilters() { return VisualLoggerFilters; }
 	void HandleTabManagerPersistLayout(const TSharedRef<FTabManager::FLayout>& LayoutToSave);
 
-	void SetFiltersPreset(const struct FFiltersPreset& Preset);
 	void OnNewWorld(UWorld* NewWorld);
 	void ResetData();
 
-protected:
-	void OnMoveCursorLeftCommand();
-	void OnMoveCursorRightCommand();
-
 protected:
 	// Holds the list of UI commands.
 	TSharedRef<FUICommandList> CommandList;
@@ -121,7 +99,7 @@ protected:
 	TSharedPtr<FTabManager> TabManager;
 
 	// Visual Logger device to get and collect logs.
-	TSharedPtr<FVisualLoggerDevice> InternalDevice;
+	TSharedPtr<FVisualLogDevice> InternalDevice;
 
 	TWeakObjectPtr<class AVisualLoggerCameraController> CameraController;
 	TSharedPtr<struct FVisualLoggerCanvasRenderer> VisualLoggerCanvasRenderer;
@@ -137,4 +115,5 @@ protected:
 	bool bGotHistogramData;
 
 	FDelegateHandle DrawOnCanvasDelegateHandle;
+	TWeakObjectPtr<class UWorld> LastUsedWorld;
 };
diff --git a/Engine/Source/Developer/LogVisualizer/Private/SVisualLoggerFilters.cpp b/Engine/Source/Developer/LogVisualizer/Private/SVisualLoggerFilters.cpp
--- a/Engine/Source/Developer/LogVisualizer/Private/SVisualLoggerFilters.cpp
+++ b/Engine/Source/Developer/LogVisualizer/Private/SVisualLoggerFilters.cpp
@@ -37,7 +37,7 @@ void SVisualLoggerFilters::Construct(const FArguments& InArgs, const TSharedRef<
 			.Text(LOCTEXT("GraphFilters", "Graph Filters"))
 		];
 
-	GraphsFilterCombo->SetVisibility(GraphFilters.Num() > 0 ? EVisibility::Visible : EVisibility::Collapsed);
+	GraphsFilterCombo->SetVisibility(CachedDatasPerGraph.Num() ? EVisibility::Visible : EVisibility::Collapsed);
 
 	FilterBox->AddSlot()
 		.Padding(3, 3)
@@ -50,7 +50,7 @@ void SVisualLoggerFilters::Construct(const FArguments& InArgs, const TSharedRef<
 			.VAlign(VAlign_Center)
 			[
 				SNew(SImage)
-				.Visibility_Lambda([this]()->EVisibility{ return this->GraphsFilter.Len() > 0 ? EVisibility::Visible : EVisibility::Hidden; })
+				.Visibility_Lambda([this]()->EVisibility{ return this->GraphsSearchString.Len() > 0 ? EVisibility::Visible : EVisibility::Hidden; })
 				.Image(FLogVisualizerStyle::Get().GetBrush("Filters.FilterIcon"))
 			]
 			+ SHorizontalBox::Slot()
@@ -59,21 +59,51 @@ void SVisualLoggerFilters::Construct(const FArguments& InArgs, const TSharedRef<
 				GraphsFilterCombo.ToSharedRef()
 			]
 		];
+
+	FVisualLoggerFilters::Get().OnFilterCategoryAdded.AddRaw(this, &SVisualLoggerFilters::OnFilterCategoryAdded);
+	FVisualLoggerFilters::Get().OnFilterCategoryRemoved.AddRaw(this, &SVisualLoggerFilters::OnFilterCategoryRemoved);
+	FVisualLoggerDatabase::Get().GetEvents().OnItemSelectionChanged.AddRaw(this, &SVisualLoggerFilters::OnItemsSelectionChanged);
+
+	FVisualLoggerDatabase::Get().GetEvents().OnGraphAddedEvent.AddRaw(this, &SVisualLoggerFilters::OnGraphAddedEvent);
+	FVisualLoggerDatabase::Get().GetEvents().OnGraphDataNameAddedEvent.AddRaw(this, &SVisualLoggerFilters::OnGraphDataNameAddedEvent);
+}
+
+SVisualLoggerFilters::~SVisualLoggerFilters()
+{
+	FVisualLoggerFilters::Get().OnFilterCategoryAdded.RemoveAll(this);
+	FVisualLoggerFilters::Get().OnFilterCategoryRemoved.RemoveAll(this);
+	FVisualLoggerDatabase::Get().GetEvents().OnItemSelectionChanged.RemoveAll(this);
+
+	FVisualLoggerDatabase::Get().GetEvents().OnGraphAddedEvent.RemoveAll(this);
+	FVisualLoggerDatabase::Get().GetEvents().OnGraphDataNameAddedEvent.RemoveAll(this);
+}
+
+void SVisualLoggerFilters::ResetData()
+{
+	for (auto& CurrentFilter : Filters)
+	{
+		FilterBox->RemoveSlot(CurrentFilter);
+	}
+	Filters.Reset();
+	CachedDatasPerGraph.Reset();
+
+	GraphsSearchString.Reset();
+	CachedGraphFilters.Reset();
 }
 
 bool SVisualLoggerFilters::GraphSubmenuVisibility(const FName MenuName)
 {
-	if (GraphsFilter.Len() == 0)
+	if (GraphsSearchString.Len() == 0)
 	{
 		return true;
 	}
 
-	if (GraphFilters.Find(MenuName) != NULL)
+	if (CachedDatasPerGraph.Contains(MenuName) != NULL)
 	{
-		const TArray<FString> DataNames = *GraphFilters.Find(MenuName);
-		for (const auto& CurrentData : DataNames)
+		const TArray<FName>& DataNames = CachedDatasPerGraph[MenuName];
+		for (const FName& CurrentData : DataNames)
 		{
-			if (CurrentData.Find(GraphsFilter) != INDEX_NONE)
+			if (CurrentData.ToString().Find(GraphsSearchString) != INDEX_NONE)
 			{
 				return true;
 			}
@@ -83,6 +113,16 @@ bool SVisualLoggerFilters::GraphSubmenuVisibility(const FName MenuName)
 	return false;
 }
 
+void SVisualLoggerFilters::OnGraphAddedEvent(const FName& OwnerName, const FName& GraphName)
+{
+	CachedDatasPerGraph.FindOrAdd(GraphName);
+}
+
+void SVisualLoggerFilters::OnGraphDataNameAddedEvent(const FName& OwnerName, const FName& GraphName, const FName& DataName)
+{
+	CachedDatasPerGraph.FindOrAdd(GraphName).AddUnique(DataName);
+}
+
 TSharedRef<SWidget> SVisualLoggerFilters::MakeGraphsFilterMenu()
 {
 	FMenuBuilder MenuBuilder(true, NULL);
@@ -90,31 +130,51 @@ TSharedRef<SWidget> SVisualLoggerFilters::MakeGraphsFilterMenu()
 	MenuBuilder.BeginSection(TEXT("Graphs"));
 	{
 		TSharedRef<SSearchBox> FiltersSearchBox = SNew(SSearchBox)
-			.InitialText(FText::FromString(GraphsFilter))
-			.HintText(LOCTEXT("GraphsFilterSearchHint", "Quick filter"))
+			.InitialText(FText::FromString(GraphsSearchString))
+			.HintText(LOCTEXT("GraphsFilterSearchHint", "Quick find"))
 			.OnTextChanged(this, &SVisualLoggerFilters::OnSearchChanged);
 
 		MenuBuilder.AddWidget(FiltersSearchBox, LOCTEXT("FiltersSearchMenuWidget", ""));
 
-		for (auto Iter = GraphFilters.CreateIterator(); Iter; ++Iter)
+		if (CachedDatasPerGraph.Num() > 0)
 		{
-			const FText& LabelText = FText::FromString(Iter->Key.ToString());
-			MenuBuilder.AddSubMenu(
-				LabelText,
-				FText::Format(LOCTEXT("FilterByTooltipPrefix", "Filter by {0}"), LabelText),
-				FNewMenuDelegate::CreateSP(this, &SVisualLoggerFilters::CreateFiltersMenuCategoryForGraph, Iter->Key),
-				FUIAction(
-				FExecuteAction::CreateSP(this, &SVisualLoggerFilters::GraphFilterCategoryClicked, Iter->Key),
-				FCanExecuteAction(),
-				FIsActionChecked::CreateSP(this, &SVisualLoggerFilters::IsGraphFilterCategoryInUse, Iter->Key),
-				FIsActionButtonVisible::CreateLambda([this, LabelText]()->bool{return GraphSubmenuVisibility(*LabelText.ToString()); })
-				),
-				NAME_None,
-				EUserInterfaceActionType::ToggleButton
-				);
+			for (auto Iter(CachedDatasPerGraph.CreateConstIterator()); Iter; ++Iter)
+			{
+				if (Iter.Value().Num() == 0)
+				{
+					continue;
+				}
+				const FName GraphName = Iter.Key();
+
+				bool bHighlightName = false;
+				const TArray<FName> SelectedOwners = FVisualLoggerDatabase::Get().GetSelectedRows();
+				for (FName CurrentOwner : SelectedOwners)
+				{
+					bHighlightName |= FVisualLoggerGraphsDatabase::Get().ContainsGraphByName(CurrentOwner, GraphName);
+					if (bHighlightName)
+					{
+						break;
+					}
+				}
+
+				const FText& LabelText = bHighlightName ? FText::FromString(FString::Printf(TEXT("* %s"), *GraphName.ToString())) : FText::FromString(FString::Printf(TEXT("  %s"), *GraphName.ToString()));
+				MenuBuilder.AddSubMenu(
+					LabelText,
+					FText::Format(LOCTEXT("FilterByTooltipPrefix", "Filter by {0}"), LabelText),
+					FNewMenuDelegate::CreateSP(this, &SVisualLoggerFilters::CreateFiltersMenuCategoryForGraph, GraphName),
+					FUIAction(
+					FExecuteAction::CreateSP(this, &SVisualLoggerFilters::GraphFilterCategoryClicked, GraphName),
+					FCanExecuteAction(),
+					FIsActionChecked::CreateSP(this, &SVisualLoggerFilters::IsGraphFilterCategoryInUse, GraphName),
+					FIsActionButtonVisible::CreateLambda([this, GraphName]()->bool{return GraphSubmenuVisibility(GraphName); })
+					),
+					NAME_None,
+					EUserInterfaceActionType::ToggleButton
+					);
+			}
 		}
 	}
-	MenuBuilder.EndSection(); //ContentBrowserFilterBasicAsset
+	MenuBuilder.EndSection();
 
 
 	FDisplayMetrics DisplayMetrics;
@@ -134,37 +194,40 @@ TSharedRef<SWidget> SVisualLoggerFilters::MakeGraphsFilterMenu()
 		];
 }
 
-void SVisualLoggerFilters::GraphFilterCategoryClicked(FName MenuCategory)
+void SVisualLoggerFilters::GraphFilterCategoryClicked(FName GraphName)
 {
-	const bool bNewSet = !IsGraphFilterCategoryInUse(MenuCategory);
+	const bool bNewSet = !IsGraphFilterCategoryInUse(GraphName);
 
-	if (GraphFilters.Contains(MenuCategory))
+	for (auto Iter(FVisualLoggerGraphsDatabase::Get().GetConstOwnersIterator()); Iter; ++Iter)
 	{
-		bool bChanged = false;
-		for (const FString &Filter : GraphFilters[MenuCategory])
-		{
-			const FString GraphFilterName = MenuCategory.ToString() + TEXT("$") + Filter;
-			FCategoryFiltersManager::Get().GetCategory(GraphFilterName).Enabled = bNewSet;
-			bChanged = true;
-		}
-
-		if (bChanged)
+		const FName OwnerName = Iter.Key();
+		if (FVisualLoggerGraphsDatabase::Get().ContainsGraphByName(OwnerName, GraphName))
 		{
-			FLogVisualizer::Get().GetVisualLoggerEvents().OnFiltersChanged.ExecuteIfBound();
-			InvalidateCanvas();
+			FVisualLoggerGraph& Graph = FVisualLoggerGraphsDatabase::Get().GetGraphByName(OwnerName, GraphName);
+			for (auto DataIt(Graph.GetConstDataIterator()); DataIt; ++DataIt)
+			{
+				FVisualLoggerFilters::Get().DisableGraphData(GraphName, DataIt->DataName, !bNewSet);
+			}
 		}
 	}
+
+	FLogVisualizer::Get().GetEvents().OnFiltersChanged.Broadcast();
+	InvalidateCanvas();
 }
 
-bool SVisualLoggerFilters::IsGraphFilterCategoryInUse(FName MenuCategory) const
+bool SVisualLoggerFilters::IsGraphFilterCategoryInUse(FName GraphName) const
 {
 	bool bInUse = false;
-	if (GraphFilters.Contains(MenuCategory))
+	for (auto Iter(FVisualLoggerGraphsDatabase::Get().GetConstOwnersIterator()); Iter; ++Iter)
 	{
-		for (const FString& Filter : GraphFilters[MenuCategory])
+		const FName OwnerName = Iter.Key();
+		if (FVisualLoggerGraphsDatabase::Get().ContainsGraphByName(OwnerName, GraphName))
 		{
-			const FString GraphFilterName = MenuCategory.ToString() + TEXT("$") + Filter;
-			bInUse |= FCategoryFiltersManager::Get().GetCategory(GraphFilterName).Enabled;
+			FVisualLoggerGraph& Graph = FVisualLoggerGraphsDatabase::Get().GetGraphByName(OwnerName, GraphName);
+			for (auto DataIt(Graph.GetConstDataIterator()); DataIt; ++DataIt)
+			{
+				bInUse |= FVisualLoggerFilters::Get().IsGraphDataDisabled(GraphName, DataIt->DataName) == false;
+			}
 		}
 	}
 
@@ -172,73 +235,43 @@ bool SVisualLoggerFilters::IsGraphFilterCategoryInUse(FName MenuCategory) const
 }
 
 
-void SVisualLoggerFilters::CreateFiltersMenuCategoryForGraph(FMenuBuilder& MenuBuilder, FName MenuCategory) const
-{
-	auto FiltersFromGraph = GraphFilters[MenuCategory];
-	for (auto Iter = FiltersFromGraph.CreateIterator(); Iter; ++Iter)
+void SVisualLoggerFilters::CreateFiltersMenuCategoryForGraph(FMenuBuilder& MenuBuilder, FName GraphName) const
+{	
+	for (auto Iter(CachedDatasPerGraph[GraphName].CreateConstIterator()); Iter; ++Iter)
 	{
-		FName Name = **Iter;
-		const FText& LabelText = FText::FromString(Name.ToString());
+		const FName& DataName = *Iter;
+		const FText& LabelText = FText::FromString(DataName.ToString());
 		MenuBuilder.AddMenuEntry(
 			LabelText,
 			FText::Format(LOCTEXT("FilterByTooltipPrefix", "Filter by {0}"), LabelText),
 			FSlateIcon(),
 			FUIAction(
-			FExecuteAction::CreateSP(this, &SVisualLoggerFilters::FilterByTypeClicked, MenuCategory, Name),
+			FExecuteAction::CreateSP(this, &SVisualLoggerFilters::FilterByTypeClicked, GraphName, DataName),
 			FCanExecuteAction(),
-			FIsActionChecked::CreateSP(this, &SVisualLoggerFilters::IsAssetTypeActionsInUse, MenuCategory, Name),
-			FIsActionButtonVisible::CreateLambda([this, LabelText]()->bool{return this->GraphsFilter.Len() == 0 || LabelText.ToString().Find(this->GraphsFilter) != INDEX_NONE; })),
+			FIsActionChecked::CreateSP(this, &SVisualLoggerFilters::IsAssetTypeActionsInUse, GraphName, DataName),
+			FIsActionButtonVisible::CreateLambda([this, LabelText]()->bool{return this->GraphsSearchString.Len() == 0 || LabelText.ToString().Find(this->GraphsSearchString) != INDEX_NONE; })),
 			NAME_None,
 			EUserInterfaceActionType::ToggleButton
 			);
 	}
 }
 
-void SVisualLoggerFilters::FilterByTypeClicked(FName InGraphName, FName InDataName)
+void SVisualLoggerFilters::FilterByTypeClicked(FName GraphName, FName DataName)
 {
-	if (GraphFilters.Contains(InGraphName))
-	{
-		bool bChanged = false;
-		for (const FString& Filter : GraphFilters[InGraphName])
-		{
-			if (Filter == InDataName.ToString())
-			{
-				const FString GraphFilterName = InGraphName.ToString() + TEXT("$") + Filter;
-				FCategoryFilter& CategoryFilter = FCategoryFiltersManager::Get().GetCategory(GraphFilterName);
-				CategoryFilter.Enabled = !CategoryFilter.Enabled;
-				bChanged = true;
-			}
-		}
-
-		if (bChanged)
-		{
-			FLogVisualizer::Get().GetVisualLoggerEvents().OnFiltersChanged.ExecuteIfBound();
-			InvalidateCanvas();
-		}
-	}
+	const bool bIsGraphDataDisabled = FVisualLoggerFilters::Get().IsGraphDataDisabled(GraphName, DataName);
+	FVisualLoggerFilters::Get().DisableGraphData(GraphName, DataName, !bIsGraphDataDisabled);
+	FLogVisualizer::Get().GetEvents().OnFiltersChanged.Broadcast();
+	InvalidateCanvas();
 }
 
-bool SVisualLoggerFilters::IsAssetTypeActionsInUse(FName InGraphName, FName InDataName) const
+bool SVisualLoggerFilters::IsAssetTypeActionsInUse(FName GraphName, FName DataName) const
 {
-	if (GraphFilters.Contains(InGraphName))
-	{
-		for (const FString& Filter : GraphFilters[InGraphName])
-		{
-			if (Filter == InDataName.ToString())
-			{
-				const FString GraphFilterName = InGraphName.ToString() + TEXT("$") + Filter;
-				return FCategoryFiltersManager::Get().GetCategory(GraphFilterName).Enabled;
-			}
-		}
-	}
-
-	return false;
+	return FVisualLoggerFilters::Get().IsGraphDataDisabled(GraphName, DataName) == false;
 }
 
 void SVisualLoggerFilters::OnSearchChanged(const FText& Filter)
 {
-	GraphsFilter = Filter.ToString();
-
+	GraphsSearchString = Filter.ToString();
 	InvalidateCanvas();
 }
 
@@ -275,55 +308,68 @@ uint32 SVisualLoggerFilters::GetCategoryIndex(const FString& InFilterName) const
 	return INDEX_NONE;
 }
 
-void SVisualLoggerFilters::AddFilter(const FString& InFilterName)
+void SVisualLoggerFilters::OnFilterCategoryAdded(FString InName, ELogVerbosity::Type InVerbosity)
 {
-	for (auto& CurrentFilter : Filters)
+	int32 CharIndex = INDEX_NONE;
+	if (InName.FindChar('$', CharIndex) == true)
 	{
-		if (CurrentFilter->GetFilterNameAsString() == InFilterName)
+		TArray<FString> GroupAndName;
+		InName.ParseIntoArray(GroupAndName, TEXT("$"), true);
+		if (ensure(GroupAndName.Num() == 2))
 		{
-			return;
+			CachedGraphFilters.FindOrAdd(*GroupAndName[0]).AddUnique(GroupAndName[1]);
 		}
 	}
+	else
+	{
+		for (auto& CurrentFilter : Filters)
+		{
+			if (CurrentFilter->GetFilterNameAsString() == InName)
+			{
+				return;
+			}
+		}
 
-	const FLinearColor Color = FLogVisualizer::Get().GetColorForCategory(InFilterName);
-	TSharedRef<SFilterWidget> NewFilter =
-		SNew(SFilterWidget)
-		.FilterName(*InFilterName)
-		.ColorCategory(Color)
-		.OnFilterChanged(SFilterWidget::FOnSimpleRequest::CreateRaw(this, &SVisualLoggerFilters::OnFiltersChanged))
-		//.OnRequestRemove(this, &SLogFilterList::RemoveFilter)
-		//.OnRequestDisableAll(this, &SLogFilterList::DisableAllFilters)
-		//.OnRequestRemoveAll(this, &SLogFilterList::RemoveAllFilters);
-		;
-
-	bool bIsSet = true;
+		const FLinearColor Color = FLogVisualizer::Get().GetColorForCategory(InName);
+		TSharedRef<SFilterWidget> NewFilter =
+			SNew(SFilterWidget)
+			.FilterName(*InName)
+			.ColorCategory(Color)
+			.OnFilterChanged(SFilterWidget::FOnSimpleRequest::CreateRaw(this, &SVisualLoggerFilters::OnFiltersChanged));
 
-	if (FCategoryFiltersManager::Get().IsValidCategory(InFilterName))
-	{
-		bIsSet = FCategoryFiltersManager::Get().GetCategory(InFilterName).Enabled;
+		Filters.Add(NewFilter);
+		FilterBox->AddSlot()
+			.Padding(2, 2)
+			[
+				NewFilter
+			];
 	}
-	FCategoryFiltersManager::Get().AddCategory(InFilterName, ELogVerbosity::All);
-	Filters.Add(NewFilter);
 
-	FilterBox->AddSlot()
-		.Padding(2, 2)
-		[
-			NewFilter
-		];
-
-	GraphsFilterCombo->SetVisibility(GraphFilters.Num() > 0 ? EVisibility::Visible : EVisibility::Collapsed);
+	GraphsFilterCombo->SetVisibility(CachedGraphFilters.Num() > 0 ? EVisibility::Visible : EVisibility::Collapsed);
 }
 
-void SVisualLoggerFilters::ResetData()
+void SVisualLoggerFilters::OnFilterCategoryRemoved(FString InName)
 {
-	for (auto& CurrentFilter : Filters)
+	int32 CharIndex = INDEX_NONE;
+	if (InName.FindChar('$', CharIndex) == true)
 	{
-		FilterBox->RemoveSlot(CurrentFilter);
+		//FIXME: Removing filter not implemented yet
+		ensureMsgf(0, TEXT("Removing filter not implemented yet"));
+	}
+	else
+	{
+		for (TSharedRef<SFilterWidget> CurrentFilter : Filters)
+		{
+			if (CurrentFilter->GetFilterNameAsString() == InName)
+			{
+				FilterBox->RemoveSlot(CurrentFilter);
+				Filters.Remove(CurrentFilter);
+				break;
+			}
+		}
 	}
-	Filters.Reset();
 }
 
-
 void SVisualLoggerFilters::OnFiltersChanged()
 {
 	TArray<FString> EnabledFilters;
@@ -335,15 +381,7 @@ void SVisualLoggerFilters::OnFiltersChanged()
 		}
 	}
 
-	FLogVisualizer::Get().GetVisualLoggerEvents().OnFiltersChanged.ExecuteIfBound();
-}
-
-void SVisualLoggerFilters::AddFilter(const FString& GraphName, const FString& DataName)
-{
-	GraphFilters.FindOrAdd(*GraphName).AddUnique(DataName);
-
-	FString GraphFilterName = GraphName + TEXT("$") + DataName;
-	FCategoryFiltersManager::Get().AddCategory(GraphFilterName, ELogVerbosity::All);
+	FLogVisualizer::Get().GetEvents().OnFiltersChanged.Broadcast();
 }
 
 void SVisualLoggerFilters::OnFiltersSearchChanged(const FText& Filter)
@@ -351,10 +389,19 @@ void SVisualLoggerFilters::OnFiltersSearchChanged(const FText& Filter)
 
 }
 
-void SVisualLoggerFilters::OnItemSelectionChanged(const struct FVisualLogEntry& EntryItem)
+void SVisualLoggerFilters::OnItemsSelectionChanged(const FVisualLoggerDBRow& ChangedRow, int32 SelectedItemIndex)
 {
 	TArray<FVisualLoggerCategoryVerbosityPair> Categories;
-	FVisualLoggerHelpers::GetCategories(EntryItem, Categories);
+	const TArray<FName>& SelectedRows = FVisualLoggerDatabase::Get().GetSelectedRows();
+	for (auto& RowName : SelectedRows)
+	{
+		FVisualLoggerDBRow& DBRow = FVisualLoggerDatabase::Get().GetRowByName(RowName);
+		if (DBRow.GetCurrentItemIndex() != INDEX_NONE)
+		{
+			FVisualLoggerHelpers::GetCategories(DBRow.GetCurrentItem().Entry, Categories);
+		}
+	}
+
 	for (int32 Index = 0; Index < Filters.Num(); ++Index)
 	{
 		SFilterWidget& Filter = Filters[Index].Get();
diff --git a/Engine/Source/Developer/LogVisualizer/Private/SVisualLoggerFilters.h b/Engine/Source/Developer/LogVisualizer/Private/SVisualLoggerFilters.h
--- a/Engine/Source/Developer/LogVisualizer/Private/SVisualLoggerFilters.h
+++ b/Engine/Source/Developer/LogVisualizer/Private/SVisualLoggerFilters.h
@@ -12,9 +12,9 @@ public:
 	SLATE_BEGIN_ARGS(SVisualLoggerFilters) { }
 	SLATE_END_ARGS()
 		
+	virtual ~SVisualLoggerFilters();
 	void Construct(const FArguments& InArgs, const TSharedRef<FUICommandList>& InCommandList);
-	void AddFilter(const FString& InFilterName);
-	void AddFilter(const FString& GraphName, const FString& DataName);
+
 	uint32 GetCategoryIndex(const FString& InFilterName) const;
 	void OnFiltersSearchChanged(const FText& Filter);
 	void OnSearchChanged(const FText& Filter);
@@ -30,12 +30,26 @@ public:
 
 	TSharedRef<SWidget> MakeGraphsFilterMenu();
 	void CreateFiltersMenuCategoryForGraph(FMenuBuilder& MenuBuilder, FName MenuCategory) const;
-	bool HasFilters() { return Filters.Num() || GraphFilters.Num(); }
+	bool HasFilters() { return Filters.Num() || CachedGraphFilters.Num(); }
+
+protected:
+	//void OnNewItemHandler(const FVisualLoggerDBRow& BDRow, int32 ItemIndex);
+
+	void OnFilterCategoryAdded(FString, ELogVerbosity::Type);
+	void OnFilterCategoryRemoved(FString);
+	void OnItemsSelectionChanged(const FVisualLoggerDBRow& ChangedRow, int32 SelectedItemIndex);
+
+	void OnGraphAddedEvent(const FName& OwnerName, const FName& GraphName);
+	void OnGraphDataNameAddedEvent(const FName& OwnerName, const FName& GraphName, const FName& DataName);
+
 protected:
 	/** The horizontal box which contains all the filters */
 	TSharedPtr<SWrapBox> FilterBox;
 	TArray<TSharedRef<SFilterWidget> > Filters;
-	TMap<FName, TArray<FString> > GraphFilters;
 	TSharedPtr<SComboButton> GraphsFilterCombo;
-	FString GraphsFilter;
+
+	FString GraphsSearchString;
+	TMap<FName, TArray<FString> > CachedGraphFilters;
+
+	TMap<FName, TArray<FName> > CachedDatasPerGraph;
 };
diff --git a/Engine/Source/Developer/LogVisualizer/Private/SVisualLoggerLogsList.cpp b/Engine/Source/Developer/LogVisualizer/Private/SVisualLoggerLogsList.cpp
--- a/Engine/Source/Developer/LogVisualizer/Private/SVisualLoggerLogsList.cpp
+++ b/Engine/Source/Developer/LogVisualizer/Private/SVisualLoggerLogsList.cpp
@@ -39,11 +39,22 @@ void SVisualLoggerLogsList::Construct(const FArguments& InArgs, const TSharedRef
 		[
 			SAssignNew(LogsLinesWidget, SListView<TSharedPtr<FLogEntryItem> >)
 			.ItemHeight(20)
-			.ListItemsSource(&LogEntryLines)
+			.ListItemsSource(&CachedLogEntryLines)
 			.SelectionMode(ESelectionMode::Multi)
 			.OnSelectionChanged(this, &SVisualLoggerLogsList::LogEntryLineSelectionChanged)
 			.OnGenerateRow(this, &SVisualLoggerLogsList::LogEntryLinesGenerateRow)
 		];
+
+	FVisualLoggerDatabase::Get().GetEvents().OnItemSelectionChanged.AddRaw(this, &SVisualLoggerLogsList::OnItemSelectionChanged);
+	FVisualLoggerDatabase::Get().GetEvents().OnRowSelectionChanged.AddRaw(this, &SVisualLoggerLogsList::ObjectSelectionChanged);
+	FLogVisualizer::Get().GetEvents().OnFiltersChanged.AddRaw(this, &SVisualLoggerLogsList::OnFiltersChanged);
+}
+
+SVisualLoggerLogsList::~SVisualLoggerLogsList()
+{
+	FVisualLoggerDatabase::Get().GetEvents().OnItemSelectionChanged.RemoveAll(this);
+	FVisualLoggerDatabase::Get().GetEvents().OnRowSelectionChanged.RemoveAll(this);
+	FLogVisualizer::Get().GetEvents().OnFiltersChanged.RemoveAll(this);
 }
 
 FReply SVisualLoggerLogsList::OnKeyDown(const FGeometry& MyGeometry, const FKeyEvent& InKeyEvent)
@@ -69,8 +80,8 @@ FReply SVisualLoggerLogsList::OnKeyDown(const FGeometry& MyGeometry, const FKeyE
 
 void SVisualLoggerLogsList::OnFiltersChanged()
 {
-	FVisualLogDevice::FVisualLogEntryItem LogEntry = CurrentLogEntry;
-	OnItemSelectionChanged(LogEntry);
+	RegenerateLogEntries();
+	LogsLinesWidget->RequestListRefresh();
 }
 
 TSharedRef<ITableRow> SVisualLoggerLogsList::LogEntryLinesGenerateRow(TSharedPtr<FLogEntryItem> Item, const TSharedRef<STableViewBase>& OwnerTable)
@@ -133,7 +144,7 @@ FText SVisualLoggerLogsList::GetFilterText() const
 {
 	static FText NoText;
 	const bool bSearchInsideLogs = ULogVisualizerSettings::StaticClass()->GetDefaultObject<ULogVisualizerSettings>()->bSearchInsideLogs;
-	return bSearchInsideLogs ? FText::FromString(FCategoryFiltersManager::Get().GetSearchString()) : NoText;
+	return bSearchInsideLogs ? FText::FromString(FVisualLoggerFilters::Get().GetSearchString()) : NoText;
 }
 
 void SVisualLoggerLogsList::OnFiltersSearchChanged(const FText& Filter)
@@ -145,69 +156,68 @@ void SVisualLoggerLogsList::LogEntryLineSelectionChanged(TSharedPtr<FLogEntryIte
 {
 	if (SelectedItem.IsValid() == true)
 	{
-		FLogVisualizer::Get().GetVisualLoggerEvents().OnLogLineSelectionChanged.ExecuteIfBound(SelectedItem, SelectedItem->UserData, SelectedItem->TagName);
+		FLogVisualizer::Get().GetEvents().OnLogLineSelectionChanged.ExecuteIfBound(SelectedItem, SelectedItem->UserData, SelectedItem->TagName);
 	}
 	else
 	{
-		FLogVisualizer::Get().GetVisualLoggerEvents().OnLogLineSelectionChanged.ExecuteIfBound(SelectedItem, 0, NAME_None);
+		FLogVisualizer::Get().GetEvents().OnLogLineSelectionChanged.ExecuteIfBound(SelectedItem, 0, NAME_None);
 	}
 }
 
-void SVisualLoggerLogsList::ObjectSelectionChanged(TArray<TSharedPtr<class STimeline> >& TimeLines)
+void SVisualLoggerLogsList::ObjectSelectionChanged(const TArray<FName>&)
 {
-	CurrentLogEntry.Entry.Reset();
-	LogEntryLines.Reset();
+	RegenerateLogEntries();
 	LogsLinesWidget->RequestListRefresh();
+}
 
-	SelectedTimeLines = TimeLines;
+void SVisualLoggerLogsList::ResetData()
+{
+	CachedLogEntryLines.Reset();
+	LogsLinesWidget->RequestListRefresh();
 }
 
-void SVisualLoggerLogsList::OnItemSelectionChanged(const FVisualLogDevice::FVisualLogEntryItem& LogEntry)
+void SVisualLoggerLogsList::OnItemSelectionChanged(const FVisualLoggerDBRow& BDRow, int32 ItemIndex)
 {
-	CurrentLogEntry.Entry.Reset();
-	LogEntryLines.Reset();
+	RegenerateLogEntries();
+}
 
-	GenerateLogs(LogEntry, SelectedTimeLines.Num() > 1);
+void SVisualLoggerLogsList::RegenerateLogEntries()
+{
+	CachedLogEntryLines.Reset();
 
-	if (SelectedTimeLines.Num() > 1)
+	const TArray<FName> SelectedRows = FVisualLoggerDatabase::Get().GetSelectedRows();
+	for (FName CurrentRow : SelectedRows)
 	{
-		for (TSharedPtr<class STimeline> CurrentTimeline : SelectedTimeLines)
+		if (FVisualLoggerDatabase::Get().IsRowVisible(CurrentRow)  == false)
 		{
-			if (CurrentTimeline.IsValid() == false || CurrentTimeline->GetVisibility().IsVisible() == false)
-			{
-				continue;
-			}
+			continue;
+		}
 
-			if (LogEntry.OwnerName == CurrentTimeline->GetName())
+		const FVisualLoggerDBRow& DBRow = FVisualLoggerDatabase::Get().GetRowByName(CurrentRow);
+		const TArray<FVisualLogDevice::FVisualLogEntryItem>& Entries = FVisualLoggerDatabase::Get().GetRowByName(CurrentRow).GetItems();
+		
+		int32 BestItemIndex = INDEX_NONE;
+		float BestDistance = MAX_FLT;
+		for (int32 Index = 0; Index < Entries.Num(); Index++)
+		{
+			auto& CurrentEntryItem = Entries[Index];
+			if (DBRow.IsItemVisible(Index) == false)
 			{
 				continue;
 			}
 
-			const TArray<FVisualLogDevice::FVisualLogEntryItem>& Entries = CurrentTimeline->GetEntries();
-
-			int32 BestItemIndex = INDEX_NONE;
-			float BestDistance = MAX_FLT;
-			for (int32 Index = 0; Index < Entries.Num(); Index++)
+			TArray<FVisualLoggerCategoryVerbosityPair> OutCategories;
+			const float CurrentDist = DBRow.GetCurrentItemIndex() == INDEX_NONE ? 0 : FMath::Abs(CurrentEntryItem.Entry.TimeStamp - DBRow.GetCurrentItem().Entry.TimeStamp);
+			if (CurrentDist < BestDistance)
 			{
-				auto& CurrentEntryItem = Entries[Index];
-
-				if (CurrentTimeline->IsEntryHidden(CurrentEntryItem))
-				{
-					continue;
-				}
-				TArray<FVisualLoggerCategoryVerbosityPair> OutCategories;
-				const float CurrentDist = FMath::Abs(CurrentEntryItem.Entry.TimeStamp - LogEntry.Entry.TimeStamp);
-				if (CurrentDist < BestDistance)
-				{
-					BestDistance = CurrentDist;
-					BestItemIndex = Index;
-				}
+				BestDistance = CurrentDist;
+				BestItemIndex = Index;
 			}
+		}
 
-			if (Entries.IsValidIndex(BestItemIndex))
-			{
-				GenerateLogs(Entries[BestItemIndex], true);
-			}
+		if (Entries.IsValidIndex(BestItemIndex))
+		{
+			GenerateLogs(Entries[BestItemIndex], SelectedRows.Num() > 1);
 		}
 	}
 }
@@ -219,7 +229,7 @@ void SVisualLoggerLogsList::GenerateLogs(const FVisualLogDevice::FVisualLogEntry
 	bool bHasValidCategory = false;
 	for (auto& CurrentCategory : OutCategories)
 	{
-		bHasValidCategory |= FCategoryFiltersManager::Get().MatchCategoryFilters(CurrentCategory.CategoryName.ToString(), CurrentCategory.Verbosity);
+		bHasValidCategory |= FVisualLoggerFilters::Get().MatchCategoryFilters(CurrentCategory.CategoryName.ToString(), CurrentCategory.Verbosity);
 	}
 	
 	if (!bHasValidCategory)
@@ -236,18 +246,17 @@ void SVisualLoggerLogsList::GenerateLogs(const FVisualLogDevice::FVisualLogEntry
 		EntryItem.Line = LogEntry.OwnerName.ToString();
 		EntryItem.UserData = 0;
 		EntryItem.TagName = NAME_None;
-		LogEntryLines.Add(MakeShareable(new FLogEntryItem(EntryItem)));
+		CachedLogEntryLines.Add(MakeShareable(new FLogEntryItem(EntryItem)));
 	}
 
-	CurrentLogEntry = LogEntry;
 	const FVisualLogLine* LogLine = LogEntry.Entry.LogLines.GetData();
 	const bool bSearchInsideLogs = ULogVisualizerSettings::StaticClass()->GetDefaultObject<ULogVisualizerSettings>()->bSearchInsideLogs;
 	for (int LineIndex = 0; LineIndex < LogEntry.Entry.LogLines.Num(); ++LineIndex, ++LogLine)
 	{
-		bool bShowLine = FCategoryFiltersManager::Get().MatchCategoryFilters(LogLine->Category.ToString(), LogLine->Verbosity);
+		bool bShowLine = FVisualLoggerFilters::Get().MatchCategoryFilters(LogLine->Category.ToString(), LogLine->Verbosity);
 		if (bSearchInsideLogs)
 		{
-			FString String = FCategoryFiltersManager::Get().GetSearchString();
+			FString String = FVisualLoggerFilters::Get().GetSearchString();
 			if (String.Len() > 0)
 			{
 				bShowLine &= LogLine->Line.Find(String) != INDEX_NONE || LogLine->Category.ToString().Find(String) != INDEX_NONE;
@@ -265,13 +274,13 @@ void SVisualLoggerLogsList::GenerateLogs(const FVisualLogDevice::FVisualLogEntry
 			EntryItem.UserData = LogLine->UserData;
 			EntryItem.TagName = LogLine->TagName;
 
-			LogEntryLines.Add(MakeShareable(new FLogEntryItem(EntryItem)));
+			CachedLogEntryLines.Add(MakeShareable(new FLogEntryItem(EntryItem)));
 		}
 	}
 
 	for (auto& Event : LogEntry.Entry.Events)
 	{
-		bool bShowLine = FCategoryFiltersManager::Get().MatchCategoryFilters(Event.Name, Event.Verbosity);
+		bool bShowLine = FVisualLoggerFilters::Get().MatchCategoryFilters(Event.Name, Event.Verbosity);
 
 		if (bShowLine)
 		{
@@ -287,7 +296,7 @@ void SVisualLoggerLogsList::GenerateLogs(const FVisualLogDevice::FVisualLogEntry
 			EntryItem.UserData = Event.UserData;
 			EntryItem.TagName = Event.TagName;
 
-			LogEntryLines.Add(MakeShareable(new FLogEntryItem(EntryItem)));
+			CachedLogEntryLines.Add(MakeShareable(new FLogEntryItem(EntryItem)));
 		}
 
 	}
diff --git a/Engine/Source/Developer/LogVisualizer/Private/SVisualLoggerLogsList.h b/Engine/Source/Developer/LogVisualizer/Private/SVisualLoggerLogsList.h
--- a/Engine/Source/Developer/LogVisualizer/Private/SVisualLoggerLogsList.h
+++ b/Engine/Source/Developer/LogVisualizer/Private/SVisualLoggerLogsList.h
@@ -9,22 +9,24 @@ public:
 	SLATE_END_ARGS()
 
 	void Construct(const FArguments& InArgs, const TSharedRef<FUICommandList>& InCommandList);
+	virtual ~SVisualLoggerLogsList();
 
-	TSharedRef<ITableRow> LogEntryLinesGenerateRow(TSharedPtr<struct FLogEntryItem> Item, const TSharedRef<STableViewBase>& OwnerTable);
-	void LogEntryLineSelectionChanged(TSharedPtr<FLogEntryItem> SelectedItem, ESelectInfo::Type SelectInfo);
 	virtual FReply OnKeyDown(const FGeometry& MyGeometry, const FKeyEvent& InKeyEvent) override;
-	void ObjectSelectionChanged(TArray<TSharedPtr<class STimeline> >& TimeLines);
 
-	void OnItemSelectionChanged(const FVisualLogDevice::FVisualLogEntryItem& EntryItem);
-	void GenerateLogs(const FVisualLogDevice::FVisualLogEntryItem& EntryItem, bool bGenerateHeader);
+	void ResetData();
+	void ObjectSelectionChanged(const TArray<FName>& Selection);
+	void OnItemSelectionChanged(const FVisualLoggerDBRow& BDRow, int32 ItemIndex);
 	void OnFiltersChanged();
 	void OnFiltersSearchChanged(const FText& Filter);
+
+protected:
 	FText GetFilterText() const;
-	const FVisualLogDevice::FVisualLogEntryItem& GetCurrentLogEntry() { return CurrentLogEntry; }
+	void RegenerateLogEntries();
+	void GenerateLogs(const FVisualLogDevice::FVisualLogEntryItem& EntryItem, bool bGenerateHeader);
+	void LogEntryLineSelectionChanged(TSharedPtr<FLogEntryItem> SelectedItem, ESelectInfo::Type SelectInfo);
+	TSharedRef<ITableRow> LogEntryLinesGenerateRow(TSharedPtr<struct FLogEntryItem> Item, const TSharedRef<STableViewBase>& OwnerTable);
 
 protected:
 	TSharedPtr<SListView<TSharedPtr<struct FLogEntryItem> > > LogsLinesWidget;
-	TArray<TSharedPtr<struct FLogEntryItem> > LogEntryLines;
-	FVisualLogDevice::FVisualLogEntryItem CurrentLogEntry;
-	TArray<TSharedPtr<class STimeline> > SelectedTimeLines;
+	TArray<TSharedPtr<struct FLogEntryItem> > CachedLogEntryLines;
 };
diff --git a/Engine/Source/Developer/LogVisualizer/Private/SVisualLoggerStatusView.cpp b/Engine/Source/Developer/LogVisualizer/Private/SVisualLoggerStatusView.cpp
--- a/Engine/Source/Developer/LogVisualizer/Private/SVisualLoggerStatusView.cpp
+++ b/Engine/Source/Developer/LogVisualizer/Private/SVisualLoggerStatusView.cpp
@@ -16,14 +16,15 @@ struct FLogStatusItem
 {
 	FString ItemText;
 	FString ValueText;
+	FString HeaderText;
 
 	TArray< TSharedPtr< FLogStatusItem > > Children;
 
+	FLogStatusItem() {}
 	FLogStatusItem(const FString& InItemText) : ItemText(InItemText) {}
 	FLogStatusItem(const FString& InItemText, const FString& InValueText) : ItemText(InItemText), ValueText(InValueText) {}
 };
 
-
 void SVisualLoggerStatusView::Construct(const FArguments& InArgs, const TSharedRef<FUICommandList>& InCommandList)
 {
 	ChildSlot
@@ -40,27 +41,33 @@ void SVisualLoggerStatusView::Construct(const FArguments& InArgs, const TSharedR
 					.TreeItemsSource(&StatusItems)
 					.OnGenerateRow(this, &SVisualLoggerStatusView::HandleGenerateLogStatus)
 					.OnGetChildren(this, &SVisualLoggerStatusView::OnLogStatusGetChildren)
+					.OnExpansionChanged(this, &SVisualLoggerStatusView::OnExpansionChanged)
 					.SelectionMode(ESelectionMode::None)
 					.Visibility(EVisibility::Visible)
 				]
-				+ SVerticalBox::Slot()
-				.FillHeight(1.0)
-				.VAlign(EVerticalAlignment::VAlign_Center)
-				[
-					SAssignNew(NotificationText, STextBlock)
-					.ColorAndOpacity(FSlateColor(FLinearColor::White))
-					.Text(FText::FromString(FString(TEXT("Multiple data selected"))))
-					.Justification(ETextJustify::Center)
-					.Visibility(EVisibility::Collapsed)
-				]
 			]
 		];
+
+	FVisualLoggerDatabase::Get().GetEvents().OnRowSelectionChanged.AddRaw(this, &SVisualLoggerStatusView::OnObjectSelectionChanged);
+	FVisualLoggerDatabase::Get().GetEvents().OnItemSelectionChanged.AddRaw(this, &SVisualLoggerStatusView::OnItemSelectionChanged);
+}
+
+SVisualLoggerStatusView::~SVisualLoggerStatusView()
+{
+	FVisualLoggerDatabase::Get().GetEvents().OnRowSelectionChanged.RemoveAll(this);
+	FVisualLoggerDatabase::Get().GetEvents().OnItemSelectionChanged.RemoveAll(this);
 }
 
-void SVisualLoggerStatusView::HideData(bool bInHide)
+void SVisualLoggerStatusView::OnExpansionChanged(TSharedPtr<FLogStatusItem> Item, bool bIsExpanded)
 {
-	NotificationText->SetVisibility(bInHide ? EVisibility::Visible : EVisibility::Collapsed);
-	StatusItemsView->SetVisibility(bInHide ? EVisibility::Collapsed : EVisibility::Visible);
+	if (bIsExpanded)
+	{
+		ExpandedCategories.AddUnique(Item->ItemText);
+	}
+	else
+	{
+		ExpandedCategories.RemoveSwap(Item->ItemText);
+	}
 }
 
 void GenerateChildren(TSharedPtr<FLogStatusItem> StatusItem, const FVisualLogStatusCategory LogCategory)
@@ -83,78 +90,61 @@ void GenerateChildren(TSharedPtr<FLogStatusItem> StatusItem, const FVisualLogSta
 	}
 }
 
-void SVisualLoggerStatusView::ObjectSelectionChanged(TArray<TSharedPtr<class STimeline> >& TimeLines)
+void SVisualLoggerStatusView::ResetData()
 {
-	SelectedTimeLines = TimeLines;
-	if (SelectedTimeLines.Num() == 0)
-	{
-		StatusItems.Reset();
-	}
+	StatusItems.Empty();
 	StatusItemsView->RequestTreeRefresh();
 }
 
-
-void SVisualLoggerStatusView::OnItemSelectionChanged(const FVisualLogDevice::FVisualLogEntryItem& LogEntry)
+void SVisualLoggerStatusView::OnObjectSelectionChanged(const TArray<FName>& SelectedItems)
 {
-
-	TArray<FString> ExpandedCategories;
+	if (SelectedItems.Num() == 0)
 	{
-		for (int32 ItemIndex = 0; ItemIndex < StatusItems.Num(); ItemIndex++)
-		{
-			const bool bIsExpanded = StatusItemsView->IsItemExpanded(StatusItems[ItemIndex]);
-			if (bIsExpanded)
-			{
-				ExpandedCategories.Add(StatusItems[ItemIndex]->ItemText);
-			}
-		}
+		ResetData();
 	}
+}
 
+void SVisualLoggerStatusView::OnItemSelectionChanged(const FVisualLoggerDBRow& ChangedDBRow, int32 ItemIndex)
+{
 	StatusItems.Empty();
+	StatusItemsView->RequestTreeRefresh();
 
-	FString TimestampDesc = FString::Printf(TEXT("%.2fs"), LogEntry.Entry.TimeStamp);
-	StatusItems.Add(MakeShareable(new FLogStatusItem(LOCTEXT("VisLogTimestamp", "Time").ToString(), TimestampDesc)));
-
-	GenerateStatusData(LogEntry, SelectedTimeLines.Num() > 1);
-
-	if (SelectedTimeLines.Num() > 1)
+	const TArray<FName>& SelectedRows = FVisualLoggerDatabase::Get().GetSelectedRows();
+	for (auto RowName : SelectedRows)
 	{
-		for (TSharedPtr<class STimeline> CurrentTimeline : SelectedTimeLines)
+		const FVisualLoggerDBRow& CurrentDBRow = FVisualLoggerDatabase::Get().GetRowByName(RowName);
+		if (FVisualLoggerDatabase::Get().IsRowVisible(RowName) == false || CurrentDBRow.GetCurrentItemIndex() == INDEX_NONE)
 		{
-			if (CurrentTimeline.IsValid() == false || CurrentTimeline->GetVisibility().IsVisible() == false)
-			{
-				continue;
-			}
-
-			if (LogEntry.OwnerName == CurrentTimeline->GetName())
-			{
-				continue;
-			}
-
-			const TArray<FVisualLogDevice::FVisualLogEntryItem>& Entries = CurrentTimeline->GetEntries();
-
-			int32 BestItemIndex = INDEX_NONE;
-			float BestDistance = MAX_FLT;
-			for (int32 Index = 0; Index < Entries.Num(); Index++)
-			{
-				auto& CurrentEntryItem = Entries[Index];
-
-				if (CurrentTimeline->IsEntryHidden(CurrentEntryItem))
-				{
-					continue;
-				}
-				TArray<FVisualLoggerCategoryVerbosityPair> OutCategories;
-				const float CurrentDist = FMath::Abs(CurrentEntryItem.Entry.TimeStamp - LogEntry.Entry.TimeStamp);
-				if (CurrentDist < BestDistance)
-				{
-					BestDistance = CurrentDist;
-					BestItemIndex = Index;
-				}
-			}
+			continue;
+		}
 
-			if (Entries.IsValidIndex(BestItemIndex))
-			{
-				GenerateStatusData(Entries[BestItemIndex], true);
-			}
+		const TArray<FVisualLogDevice::FVisualLogEntryItem>& Entries = CurrentDBRow.GetItems();
+		//int32 BestItemIndex = INDEX_NONE;
+		//float BestDistance = MAX_FLT;
+		//for (int32 Index = 0; Index < Entries.Num(); Index++)
+		//{
+		//	auto& CurrentEntryItem = Entries[Index];
+
+		//	if (CurrentDBRow.IsItemVisible(Index) == false)
+		//	{
+		//		continue;
+		//	}
+
+		//	TArray<FVisualLoggerCategoryVerbosityPair> OutCategories;
+		//	const float LastSelectedTimestep = CurrentDBRow.GetCurrentItemIndex() != INDEX_NONE ? CurrentDBRow.GetCurrentItem().Entry.TimeStamp : CurrentEntryItem.Entry.TimeStamp;
+		//	const float CurrentDist = FMath::Abs(CurrentEntryItem.Entry.TimeStamp - LastSelectedTimestep);
+		//	if (CurrentDist < BestDistance)
+		//	{
+		//		BestDistance = CurrentDist;
+		//		BestItemIndex = Index;
+		//	}
+		//}
+		
+
+		//if (Entries.IsValidIndex(BestItemIndex))
+		{
+			//GenerateStatusData(Entries[BestItemIndex], SelectedRows.Num() > 1);
+			GenerateStatusData(CurrentDBRow.GetCurrentItem(), SelectedRows.Num() > 1);
 		}
 	}
 
@@ -176,7 +166,14 @@ void SVisualLoggerStatusView::GenerateStatusData(const FVisualLogDevice::FVisual
 {
 	if (bAddHeader)
 	{
-		StatusItems.Add(MakeShareable(new FLogStatusItem(LOCTEXT("VisLogOwnerName", "Owner Name").ToString(), LogEntry.OwnerName.ToString())));
+		FLogStatusItem* HeaderItem = new FLogStatusItem();
+		HeaderItem->HeaderText = FString::Printf(TEXT("%s at Time: %.2fs"), *LogEntry.OwnerName.ToString(), LogEntry.Entry.TimeStamp);;
+		StatusItems.Add(MakeShareable(HeaderItem));
+	}
+	else
+	{
+		FString TimestampDesc = FString::Printf(TEXT("%.2fs"), LogEntry.Entry.TimeStamp);
+		StatusItems.Add(MakeShareable(new FLogStatusItem(LOCTEXT("VisLogTimestamp", "Time").ToString(), TimestampDesc)));
 	}
 
 	for (int32 CategoryIndex = 0; CategoryIndex < LogEntry.Entry.Status.Num(); CategoryIndex++)
@@ -218,6 +215,17 @@ void SVisualLoggerStatusView::OnLogStatusGetChildren(TSharedPtr<FLogStatusItem>
 
 TSharedRef<ITableRow> SVisualLoggerStatusView::HandleGenerateLogStatus(TSharedPtr<FLogStatusItem> InItem, const TSharedRef<STableViewBase>& OwnerTable)
 {
+	if (InItem->HeaderText.Len() > 0)
+	{
+		return SNew(STableRow<TSharedPtr<FLogStatusItem> >, OwnerTable)
+			.Style(&FLogVisualizerStyle::Get().GetWidgetStyle<FTableRowStyle>("TableView.DarkRow"))
+			[
+				SNew(STextBlock)
+				.Text(FText::FromString(InItem->HeaderText))
+				.ColorAndOpacity(FColorList::LightGrey)
+			];
+	}
+	
 	if (InItem->Children.Num() > 0)
 	{
 		return SNew(STableRow<TSharedPtr<FLogStatusItem> >, OwnerTable)
@@ -252,4 +260,4 @@ TSharedRef<ITableRow> SVisualLoggerStatusView::HandleGenerateLogStatus(TSharedPt
 		];
 }
 
-#undef LOCTEXT_NAMESPACE
+#undef LOCTEXT_NAMESPACE
\ No newline at end of file
diff --git a/Engine/Source/Developer/LogVisualizer/Private/SVisualLoggerStatusView.h b/Engine/Source/Developer/LogVisualizer/Private/SVisualLoggerStatusView.h
--- a/Engine/Source/Developer/LogVisualizer/Private/SVisualLoggerStatusView.h
+++ b/Engine/Source/Developer/LogVisualizer/Private/SVisualLoggerStatusView.h
@@ -8,19 +8,21 @@ public:
 	SLATE_BEGIN_ARGS(SVisualLoggerStatusView){}
 	SLATE_END_ARGS();
 
+	virtual ~SVisualLoggerStatusView();
 	void Construct(const FArguments& InArgs, const TSharedRef<FUICommandList>& InCommandList);
+	void ResetData();
 
-	void OnItemSelectionChanged(const FVisualLogDevice::FVisualLogEntryItem&);
+protected:
+	void OnObjectSelectionChanged(const TArray<FName>& SelectedRows);
+	void OnItemSelectionChanged(const FVisualLoggerDBRow& DBRow, int32 ItemIndex);
 	void GenerateStatusData(const FVisualLogDevice::FVisualLogEntryItem&, bool bAddHeader);
-	void ObjectSelectionChanged(TArray<TSharedPtr<class STimeline> >& TimeLines);
 
 	TSharedRef<ITableRow> HandleGenerateLogStatus(TSharedPtr<struct FLogStatusItem> InItem, const TSharedRef<STableViewBase>& OwnerTable);
 	void OnLogStatusGetChildren(TSharedPtr<struct FLogStatusItem> InItem, TArray< TSharedPtr<FLogStatusItem> >& OutItems);
-	void HideData(bool bInHide);
+	void OnExpansionChanged(TSharedPtr<FLogStatusItem> Item, bool);
 
 protected:
+	TArray<FString> ExpandedCategories;
 	TSharedPtr< STreeView< TSharedPtr<FLogStatusItem> > > StatusItemsView;
-	TSharedPtr< STextBlock > NotificationText;
 	TArray< TSharedPtr<FLogStatusItem> > StatusItems;
-	TArray<TSharedPtr<class STimeline> > SelectedTimeLines;
 };
diff --git a/Engine/Source/Developer/LogVisualizer/Private/SVisualLoggerView.cpp b/Engine/Source/Developer/LogVisualizer/Private/SVisualLoggerView.cpp
--- a/Engine/Source/Developer/LogVisualizer/Private/SVisualLoggerView.cpp
+++ b/Engine/Source/Developer/LogVisualizer/Private/SVisualLoggerView.cpp
@@ -64,21 +64,9 @@ FReply SInputCatcherOverlay::OnMouseWheel(const FGeometry& MyGeometry, const FPo
 	return FReply::Unhandled();
 }
 
-void SVisualLoggerView::GetTimelines(TArray<TSharedPtr<STimeline> >& OutList, bool bOnlySelectedOnes)
-{
-	OutList = bOnlySelectedOnes ? TimelinesContainer->GetSelectedNodes() : TimelinesContainer->GetAllNodes();
-}
-
 void SVisualLoggerView::Construct(const FArguments& InArgs, const TSharedRef<FUICommandList>& InCommandList)
 {
 	AnimationOutlinerFillPercentage = .25f;
-	VisualLoggerEvents = FLogVisualizer::Get().GetVisualLoggerEvents();
-
-	FVisualLoggerTimeSliderArgs TimeSliderArgs;
-	TimeSliderArgs.ViewRange = FAnimatedRange::WrapAttribute(InArgs._ViewRange);
-	TimeSliderArgs.ClampRange = TRange<float>(InArgs._ViewRange.Get().GetLowerBoundValue(), InArgs._ViewRange.Get().GetUpperBoundValue());
-	TimeSliderArgs.ScrubPosition = InArgs._ScrubPosition;
-	FLogVisualizer::Get().GetTimeSliderController()->SetTimesliderArgs(TimeSliderArgs);
 
 	TSharedRef<SScrollBar> ZoomScrollBar =
 		SNew(SScrollBar)
@@ -116,6 +104,7 @@ void SVisualLoggerView::Construct(const FArguments& InArgs, const TSharedRef<FUI
 						SAssignNew(SearchSplitter, SSplitter)
 						.Orientation(Orient_Horizontal)
 						.OnSplitterFinishedResizing(this, &SVisualLoggerView::OnSearchSplitterResized)
+						
 						+ SSplitter::Slot()
 						.Value(0.25)
 						[
@@ -127,7 +116,7 @@ void SVisualLoggerView::Construct(const FArguments& InArgs, const TSharedRef<FUI
 							.VAlign(VAlign_Center)
 							[
 								SNew(SImage)
-								.Visibility_Lambda([]()->EVisibility{ return FCategoryFiltersManager::Get().GetSelectedObjects().Num() > 0 ? EVisibility::Visible : EVisibility::Collapsed; })
+								.Visibility_Lambda([]()->EVisibility{ return FVisualLoggerFilters::Get().GetSelectedObjects().Num() > 0 ? EVisibility::Visible : EVisibility::Collapsed; })
 								.Image(FLogVisualizerStyle::Get().GetBrush("Filters.FilterIcon"))
 							]
 							+ SHorizontalBox::Slot()
@@ -136,7 +125,7 @@ void SVisualLoggerView::Construct(const FArguments& InArgs, const TSharedRef<FUI
 							.AutoWidth()
 							[
 								SAssignNew(ClassesComboButton, SComboButton)
-								.Visibility_Lambda([this]()->EVisibility{ return TimelinesContainer.IsValid() && (TimelinesContainer->GetAllNodes().Num() > 1 || FCategoryFiltersManager::Get().GetSelectedObjects().Num() > 0) ? EVisibility::Visible : EVisibility::Collapsed; })
+								.Visibility_Lambda([this]()->EVisibility{ return TimelinesContainer.IsValid() && (TimelinesContainer->GetAllNodes().Num() > 1 || FVisualLoggerFilters::Get().GetSelectedObjects().Num() > 0) ? EVisibility::Visible : EVisibility::Collapsed; })
 								.ComboButtonStyle(FLogVisualizerStyle::Get(), "Filters.Style")
 								.ForegroundColor(FLinearColor::White)
 								.ContentPadding(0)
@@ -261,7 +250,21 @@ void SVisualLoggerView::Construct(const FArguments& InArgs, const TSharedRef<FUI
 			]
 		];
 		
-		SearchBox->SetText(FText::FromString(FCategoryFiltersManager::Get().GetSearchString()));
+	SearchBox->SetText(FText::FromString(FVisualLoggerFilters::Get().GetSearchString()));
+	FLogVisualizer::Get().GetEvents().GetAnimationOutlinerFillPercentageFunc.BindLambda(
+		[this]()->float{ 
+			SSplitter::FSlot const& LeftSplitterSlot = SearchSplitter->SlotAt(0);
+			SSplitter::FSlot const& RightSplitterSlot = SearchSplitter->SlotAt(1);
+
+			return LeftSplitterSlot.SizeValue.Get() / RightSplitterSlot.SizeValue.Get();
+	});
+
+	OnSearchSplitterResized();
+}
+
+SVisualLoggerView::~SVisualLoggerView()
+{
+	FLogVisualizer::Get().GetEvents().GetAnimationOutlinerFillPercentageFunc.Unbind();
 }
 
 void SVisualLoggerView::SetAnimationOutlinerFillPercentage(float FillPercentage) 
@@ -277,23 +280,16 @@ void SVisualLoggerView::SetSearchString(FText SearchString)
 	}
 }
 
-void SVisualLoggerView::OnObjectSelectionChanged(TSharedPtr<class STimeline> TimeLine)
-{
-	//FIXME: scroll to selected timeline (SebaK)
-	//FWidgetPath WidgetPath;
-	//if (FSlateApplication::Get().GeneratePathToWidgetUnchecked(TimeLine.ToSharedRef(), WidgetPath))
-	//{
-	//	FArrangedWidget ArrangedWidget = WidgetPath.FindArrangedWidget(TimeLine.ToSharedRef()).Get(FArrangedWidget::NullWidget);
-	//	ScrollBox->ScrollDescendantIntoView(ArrangedWidget.Geometry, TimeLine, true);
-	//}
-}
 
 void SVisualLoggerView::OnSearchSplitterResized()
 {
 	SSplitter::FSlot const& LeftSplitterSlot = SearchSplitter->SlotAt(0);
 	SSplitter::FSlot const& RightSplitterSlot = SearchSplitter->SlotAt(1);
 
-	SetAnimationOutlinerFillPercentage(LeftSplitterSlot.SizeValue.Get() / RightSplitterSlot.SizeValue.Get());
+	const float AnimationOutlinerFillPercentage = LeftSplitterSlot.SizeValue.Get() / RightSplitterSlot.SizeValue.Get();
+	SetAnimationOutlinerFillPercentage(AnimationOutlinerFillPercentage);
+
+	FLogVisualizer::Get().SetAnimationOutlinerFillPercentage(AnimationOutlinerFillPercentage);
 }
 
 void SVisualLoggerView::OnSearchChanged(const FText& Filter)
@@ -326,11 +322,6 @@ void SVisualLoggerView::ResetData()
 	TimelinesContainer->ResetData();
 }
 
-void SVisualLoggerView::OnNewLogEntry(const FVisualLogDevice::FVisualLogEntryItem& Entry)
-{
-	TimelinesContainer->OnNewLogEntry(Entry);
-}
-
 void SVisualLoggerView::OnFiltersChanged()
 {
 	TimelinesContainer->OnFiltersChanged();
@@ -371,13 +362,13 @@ TSharedRef<SWidget> SVisualLoggerView::MakeClassesFilterMenu()
 				FUIAction(
 				FExecuteAction::CreateLambda([this, OwnerClassName]()
 				{
-				if (FCategoryFiltersManager::Get().MatchObjectName(OwnerClassName) && FCategoryFiltersManager::Get().GetSelectedObjects().Num() != 0)
+				if (FVisualLoggerFilters::Get().MatchObjectName(OwnerClassName) && FVisualLoggerFilters::Get().GetSelectedObjects().Num() != 0)
 					{
-						FCategoryFiltersManager::Get().RemoveObjectFromSelection(OwnerClassName);
+						FVisualLoggerFilters::Get().RemoveObjectFromSelection(OwnerClassName);
 					}
 					else
 					{
-						FCategoryFiltersManager::Get().SelectObject(OwnerClassName);
+						FVisualLoggerFilters::Get().SelectObject(OwnerClassName);
 					}
 
 					OnChangedClassesFilter();
@@ -385,7 +376,7 @@ TSharedRef<SWidget> SVisualLoggerView::MakeClassesFilterMenu()
 				FCanExecuteAction(),
 				FIsActionChecked::CreateLambda([OwnerClassName]()->bool
 				{
-					return FCategoryFiltersManager::Get().GetSelectedObjects().Find(OwnerClassName) != INDEX_NONE;
+					return FVisualLoggerFilters::Get().GetSelectedObjects().Find(OwnerClassName) != INDEX_NONE;
 				}),
 				FIsActionButtonVisible()),
 				NAME_None,
@@ -395,7 +386,7 @@ TSharedRef<SWidget> SVisualLoggerView::MakeClassesFilterMenu()
 		}
 	}
 	//show any classes from persistent data
-	for (const FString& SelectedObj : FCategoryFiltersManager::Get().GetSelectedObjects())
+	for (const FString& SelectedObj : FVisualLoggerFilters::Get().GetSelectedObjects())
 	{
 		if (UniqueClasses.Find(SelectedObj) == INDEX_NONE)
 		{
@@ -407,13 +398,13 @@ TSharedRef<SWidget> SVisualLoggerView::MakeClassesFilterMenu()
 			FUIAction(
 			FExecuteAction::CreateLambda([this, SelectedObj]()
 			{
-				if (FCategoryFiltersManager::Get().MatchObjectName(SelectedObj) && FCategoryFiltersManager::Get().GetSelectedObjects().Num() != 0)
+				if (FVisualLoggerFilters::Get().MatchObjectName(SelectedObj) && FVisualLoggerFilters::Get().GetSelectedObjects().Num() != 0)
 				{
-					FCategoryFiltersManager::Get().RemoveObjectFromSelection(SelectedObj);
+					FVisualLoggerFilters::Get().RemoveObjectFromSelection(SelectedObj);
 				}
 				else
 				{
-					FCategoryFiltersManager::Get().SelectObject(SelectedObj);
+					FVisualLoggerFilters::Get().SelectObject(SelectedObj);
 				}
 
 				OnChangedClassesFilter();
@@ -421,7 +412,7 @@ TSharedRef<SWidget> SVisualLoggerView::MakeClassesFilterMenu()
 			FCanExecuteAction(),
 			FIsActionChecked::CreateLambda([SelectedObj]()->bool
 			{
-				return FCategoryFiltersManager::Get().GetSelectedObjects().Find(SelectedObj) != INDEX_NONE;
+				return FVisualLoggerFilters::Get().GetSelectedObjects().Find(SelectedObj) != INDEX_NONE;
 			}),
 			FIsActionButtonVisible()),
 			NAME_None,
diff --git a/Engine/Source/Developer/LogVisualizer/Private/SVisualLoggerView.h b/Engine/Source/Developer/LogVisualizer/Private/SVisualLoggerView.h
--- a/Engine/Source/Developer/LogVisualizer/Private/SVisualLoggerView.h
+++ b/Engine/Source/Developer/LogVisualizer/Private/SVisualLoggerView.h
@@ -16,6 +16,7 @@ public:
 		SLATE_EVENT(FOnFiltersSearchChanged, OnFiltersSearchChanged)
 	SLATE_END_ARGS();
 
+	virtual ~SVisualLoggerView();
 	void Construct(const FArguments& InArgs, const TSharedRef<FUICommandList>& InCommandList);
 	float GetAnimationOutlinerFillPercentage() const { 
 		SSplitter::FSlot const& LeftSplitterSlot = SearchSplitter->SlotAt(0);
@@ -28,15 +29,12 @@ public:
 	TSharedRef<SWidget> MakeSectionOverlay(TSharedRef<class FVisualLoggerTimeSliderController> TimeSliderController, const TAttribute< TRange<float> >& ViewRange, const TAttribute<float>& ScrubPosition, bool bTopOverlay);
 	void SetSearchString(FText SearchString);
 
-	void OnNewLogEntry(const FVisualLogDevice::FVisualLogEntryItem& Entry);
 	void OnFiltersChanged();
 	void OnSearchChanged(const FText& Filter);
 	void OnFiltersSearchChanged(const FText& Filter);
 	void OnSearchSplitterResized();
-	void OnObjectSelectionChanged(TSharedPtr<class STimeline> TimeLine);
 	void OnChangedClassesFilter();
 
-	void GetTimelines(TArray<TSharedPtr<class STimeline> >&, bool bOnlySelectedOnes = false);
 	virtual FCursorReply OnCursorQuery(const FGeometry& MyGeometry, const FPointerEvent& CursorEvent) const override;
 	TSharedRef<SWidget> MakeClassesFilterMenu();
 	void ResetData();
@@ -48,6 +46,5 @@ protected:
 	TSharedPtr<class SSearchBox> SearchBox;
 	TSharedPtr<class SComboButton> ClassesComboButton;
 
-	FVisualLoggerEvents	VisualLoggerEvents;
 	float AnimationOutlinerFillPercentage;
 };
diff --git a/Engine/Source/Developer/LogVisualizer/Private/VisualLoggerCanvasRenderer.cpp b/Engine/Source/Developer/LogVisualizer/Private/VisualLoggerCanvasRenderer.cpp
--- a/Engine/Source/Developer/LogVisualizer/Private/VisualLoggerCanvasRenderer.cpp
+++ b/Engine/Source/Developer/LogVisualizer/Private/VisualLoggerCanvasRenderer.cpp
@@ -52,21 +52,41 @@ namespace LogVisualizer
 	}
 }
 
-void FVisualLoggerCanvasRenderer::OnItemSelectionChanged(const struct FVisualLogEntry& EntryItem)
+FVisualLoggerCanvasRenderer::FVisualLoggerCanvasRenderer() 
+	: bDirtyData(true) 
 {
-	SelectedEntry = EntryItem;
+	FLogVisualizer::Get().GetEvents().OnFiltersChanged.AddRaw(this, &FVisualLoggerCanvasRenderer::DirtyCachedData);
+	FVisualLoggerDatabase::Get().GetEvents().OnRowSelectionChanged.AddRaw(this, &FVisualLoggerCanvasRenderer::ObjectSelectionChanged);
+	FVisualLoggerDatabase::Get().GetEvents().OnItemSelectionChanged.AddRaw(this, &FVisualLoggerCanvasRenderer::OnItemSelectionChanged);
+}
+
+FVisualLoggerCanvasRenderer::~FVisualLoggerCanvasRenderer()
+{
+	FLogVisualizer::Get().GetEvents().OnFiltersChanged.RemoveAll(this);
+	FVisualLoggerDatabase::Get().GetEvents().OnRowSelectionChanged.RemoveAll(this);
+	FVisualLoggerDatabase::Get().GetEvents().OnItemSelectionChanged.RemoveAll(this);
+}
+
+void FVisualLoggerCanvasRenderer::ResetData()
+{
+	SelectedEntry = FVisualLogEntry();
 	DirtyCachedData();
 }
 
-void FVisualLoggerCanvasRenderer::ObjectSelectionChanged(TSharedPtr<class STimeline> TimeLine)
+void FVisualLoggerCanvasRenderer::OnItemSelectionChanged(const FVisualLoggerDBRow& ChangedRow, int32 SelectedItemIndex)
+{
+	SelectedEntry = ChangedRow.GetCurrentItemIndex() != INDEX_NONE ? ChangedRow.GetCurrentItem().Entry : FVisualLogEntry();
+	DirtyCachedData();
+}
+
+void FVisualLoggerCanvasRenderer::ObjectSelectionChanged(const TArray<FName>& RowNames)
 {
-	CurrentTimeLine = TimeLine;
 	DirtyCachedData();
 }
 
 void FVisualLoggerCanvasRenderer::DrawOnCanvas(class UCanvas* Canvas, class APlayerController*)
 {
-	if (GEngine == NULL || CurrentTimeLine.IsValid() == false)
+	if (GEngine == NULL)
 	{
 		return;
 	}
@@ -83,36 +103,79 @@ void FVisualLoggerCanvasRenderer::DrawOnCanvas(class UCanvas* Canvas, class APla
 	const FString TimeStampString = FString::Printf(TEXT("%.2f"), SelectedEntry.TimeStamp);
 	LogVisualizer::DrawTextShadowed(Canvas, Font, TimeStampString, SelectedEntry.Location);
 
-	if (bDirtyData)
+	if (bDirtyData && FLogVisualizer::Get().GetTimeSliderController().IsValid())
 	{
-		for (const auto& CurrentData : CachedDataBlocks)
+		const FVisualLoggerTimeSliderArgs&  TimeSliderArgs = FLogVisualizer::Get().GetTimeSliderController()->GetTimeSliderArgs();
+		TRange<float> LocalViewRange = TimeSliderArgs.ViewRange.Get();
+		const float LocalViewRangeMin = LocalViewRange.GetLowerBoundValue();
+		const float LocalViewRangeMax = LocalViewRange.GetUpperBoundValue();
+		const float LocalSequenceLength = LocalViewRangeMax - LocalViewRangeMin;
+		const float WindowHalfWidth = LocalSequenceLength * TimeSliderArgs.CursorSize.Get() * 0.5f;
+		const FVector2D TimeStampWindow(SelectedEntry.TimeStamp - WindowHalfWidth, SelectedEntry.TimeStamp + WindowHalfWidth);
+
+		CollectedGraphs.Reset();
+		const TArray<FName>& RowNames = FVisualLoggerDatabase::Get().GetSelectedRows();
+		for (auto RowName : RowNames)
 		{
-			if (FVisualLogExtensionInterface* Extension = FVisualLogger::Get().GetExtensionForTag(CurrentData.TagName))
+			if (FVisualLoggerDatabase::Get().IsRowVisible(RowName) == false)
 			{
-				Extension->DisableDrawingForData(FLogVisualizer::Get().GetWorld(), Canvas, NULL, CurrentData.TagName, CurrentData, SelectedEntry.TimeStamp);
+				continue;
 			}
-		}
-		CachedDataBlocks.Reset();
 
-		for (const auto& CurrentData : SelectedEntry.DataBlocks)
-		{
-			const bool bIsValidByFilter = FCategoryFiltersManager::Get().MatchCategoryFilters(CurrentData.Category.ToString(), ELogVerbosity::All) && FCategoryFiltersManager::Get().MatchCategoryFilters(CurrentData.TagName.ToString(), ELogVerbosity::All);
-			if (bIsValidByFilter)
+			const TArray<FVisualLoggerGraph>& AllGraphs = FVisualLoggerGraphsDatabase::Get().GetGraphsByOwnerName(RowName);
+			for (const FVisualLoggerGraph& CurrentGraph : AllGraphs)
 			{
-				CachedDataBlocks.Add(CurrentData);
-			}
-		}
-	}
+				const FName GraphName = CurrentGraph.GetGraphName();
+				const FName OwnerName = CurrentGraph.GetOwnerName();
 
-	for (const auto CurrentData : CachedDataBlocks)
-	{
-		FVisualLogExtensionInterface* Extension = FVisualLogger::Get().GetExtensionForTag(CurrentData.TagName);
-		if (!Extension)
-		{
-			continue;
-		}
+				if (FVisualLoggerGraphsDatabase::Get().IsGraphVisible(OwnerName, GraphName) == false)
+				{
+					continue;
+				}
+
+				for (auto DataIt(CurrentGraph.GetConstDataIterator()); DataIt; ++DataIt)
+				{
+					const FVisualLoggerGraphData& GraphData = *DataIt;
+					const bool bIsGraphDataDisabled = FVisualLoggerFilters::Get().IsGraphDataDisabled(GraphName, GraphData.DataName);
+					if (bIsGraphDataDisabled)
+					{
+						continue;
+					}
+
+					const FName FullGraphName = *FString::Printf(TEXT("%s$%s"), *RowName.ToString(), *GraphName.ToString());
+					FGraphData &CollectedGraphData = CollectedGraphs.FindOrAdd(FullGraphName);
+					FGraphLineData &LineData = CollectedGraphData.GraphLines.FindOrAdd(GraphData.DataName);
+					LineData.DataName = GraphData.DataName;
+					LineData.Samples = GraphData.Samples;
+					LineData.LeftExtreme = FVector2D::ZeroVector;
+					LineData.RightExtreme = FVector2D::ZeroVector;
+
+					int32 LeftSideOutsideIndex = INDEX_NONE;
+					int32 RightSideOutsideIndex = INDEX_NONE;
+					for (int32 SampleIndex = 0; SampleIndex < GraphData.Samples.Num(); SampleIndex++)
+					{
+						const FVector2D& SampleValue = GraphData.Samples[SampleIndex];
+						CollectedGraphData.Min.X = FMath::Min(CollectedGraphData.Min.X, SampleValue.X);
+						CollectedGraphData.Min.Y = FMath::Min(CollectedGraphData.Min.Y, SampleValue.Y);
+
+						CollectedGraphData.Max.X = FMath::Max(CollectedGraphData.Max.X, SampleValue.X);
+						CollectedGraphData.Max.Y = FMath::Max(CollectedGraphData.Max.Y, SampleValue.Y);
 
-		Extension->DrawData(FLogVisualizer::Get().GetWorld(), Canvas, NULL, CurrentData.TagName, CurrentData, SelectedEntry.TimeStamp);
+						const float CurrentTimeStamp = GraphData.TimeStamps[SampleIndex];
+						if (CurrentTimeStamp < TimeStampWindow.X)
+						{
+							LineData.LeftExtreme = SampleValue;
+						}
+						else if (CurrentTimeStamp > TimeStampWindow.Y)
+						{
+							LineData.RightExtreme = SampleValue;
+							break;
+						}
+					}
+				}
+			}
+		}
+		bDirtyData = false;
 	}
 
 	if (ULogVisualizerSessionSettings::StaticClass()->GetDefaultObject<ULogVisualizerSessionSettings>()->bEnableGraphsVisualization)
@@ -120,152 +183,31 @@ void FVisualLoggerCanvasRenderer::DrawOnCanvas(class UCanvas* Canvas, class APla
 		DrawHistogramGraphs(Canvas, NULL);
 	}
 
-	bDirtyData = false;
+	const TMap<FName, FVisualLogExtensionInterface*>& Extensions = FVisualLogger::Get().GetAllExtensions();
+	for (const auto& CurrentExtension : Extensions)
+	{
+		CurrentExtension.Value->DrawData(FVisualLoggerEditorInterface::Get(), Canvas);
+	}
 }
 
 void FVisualLoggerCanvasRenderer::DrawHistogramGraphs(class UCanvas* Canvas, class APlayerController*)
 {
-	struct FGraphLineData
-	{
-		FName DataName;
-		FVector2D LeftExtreme, RightExtreme;
-		TArray<FVector2D> Samples;
-	};
-
-	struct FGraphData
-	{
-		FGraphData() : Min(FVector2D(FLT_MAX, FLT_MAX)), Max(FVector2D(FLT_MIN, FLT_MIN)) {}
-
-		FVector2D Min, Max;
-		TMap<FName, FGraphLineData> GraphLines;
-	};
-
 	if (FLogVisualizer::Get().GetTimeSliderController().IsValid() == false)
 	{
 		return;
 	}
 
-	TMap<FName, FGraphData>	CollectedGraphs;
-
-	const FVisualLoggerTimeSliderArgs&  TimeSliderArgs = FLogVisualizer::Get().GetTimeSliderController()->GetTimeSliderArgs();
-	TRange<float> LocalViewRange = TimeSliderArgs.ViewRange.Get();
-	const float LocalViewRangeMin = LocalViewRange.GetLowerBoundValue();
-	const float LocalViewRangeMax = LocalViewRange.GetUpperBoundValue();
-	const float LocalSequenceLength = LocalViewRangeMax - LocalViewRangeMin;
-	const float WindowHalfWidth = LocalSequenceLength * TimeSliderArgs.CursorSize.Get() * 0.5f;
-	const FVector2D TimeStampWindow(SelectedEntry.TimeStamp - WindowHalfWidth, SelectedEntry.TimeStamp + WindowHalfWidth);
-
-	const TArray<FVisualLogDevice::FVisualLogEntryItem> &ObjectItems = CurrentTimeLine.Pin()->GetEntries();
-	int32 ColorIndex = 0;
-	int32 LeftSideOutsideIndex = INDEX_NONE;
-	int32 RightSideOutsideIndex = INDEX_NONE;
-
-	if (bDirtyData)
-	{
-		CachedHistogramSamples.Reset();
-		TArray<FVisualLogHistogramSample> CurrentSamples;
-		for (int32 EntryIndex = 0; EntryIndex < ObjectItems.Num(); ++EntryIndex)
-		{
-			const FVisualLogEntry* CurrentEntry = &(ObjectItems[EntryIndex].Entry);
-			if (CurrentEntry->TimeStamp < TimeStampWindow.X)
-			{
-				LeftSideOutsideIndex = EntryIndex;
-				continue;
-			}
-
-			if (CurrentEntry->TimeStamp > TimeStampWindow.Y)
-			{
-				RightSideOutsideIndex = EntryIndex;
-				break;
-			}
-
-			CurrentSamples.Reset();
-			const int32 SamplesNum = CurrentEntry->HistogramSamples.Num();
-			for (int32 SampleIndex = 0; SampleIndex < SamplesNum; ++SampleIndex)
-			{
-				const FVisualLogHistogramSample& CurrentSample = CurrentEntry->HistogramSamples[SampleIndex];
-
-				const FName CurrentCategory = CurrentSample.Category;
-				const FName CurrentGraphName = CurrentSample.GraphName;
-				const FName CurrentDataName = CurrentSample.DataName;
-
-				FString GraphFilterName = CurrentSample.GraphName.ToString() + TEXT("$") + CurrentSample.DataName.ToString();
-				const bool bIsValidByFilter = FCategoryFiltersManager::Get().MatchCategoryFilters(GraphFilterName, ELogVerbosity::All);
-				if (bIsValidByFilter)
-				{
-					CurrentSamples.Add(CurrentSample);
-				}
-			}
-
-			if (CurrentSamples.Num() > 0)
-			{
-				CachedHistogramSamples.Add(CurrentSamples);
-			}
-		}
-	}
-
-	for (int32 EntryIndex = 0; EntryIndex < CachedHistogramSamples.Num(); ++EntryIndex)
-	{
-		TArray<FVisualLogHistogramSample> &HistogramSamples = CachedHistogramSamples[EntryIndex];
-		const int32 SamplesNum = CachedHistogramSamples[EntryIndex].Num();
-		for (int32 SampleIndex = 0; SampleIndex < SamplesNum; ++SampleIndex)
-		{
-			FVisualLogHistogramSample CurrentSample = HistogramSamples[SampleIndex];
-
-			FGraphData &GraphData = CollectedGraphs.FindOrAdd(CurrentSample.GraphName);
-			FGraphLineData &LineData = GraphData.GraphLines.FindOrAdd(CurrentSample.DataName);
-			LineData.DataName = CurrentSample.DataName;
-			LineData.Samples.Add(CurrentSample.SampleValue);
-			LineData.LeftExtreme = FVector2D::ZeroVector;
-			LineData.RightExtreme = FVector2D::ZeroVector;
-
-			GraphData.Min.X = FMath::Min(GraphData.Min.X, CurrentSample.SampleValue.X);
-			GraphData.Min.Y = FMath::Min(GraphData.Min.Y, CurrentSample.SampleValue.Y);
-
-			GraphData.Max.X = FMath::Max(GraphData.Max.X, CurrentSample.SampleValue.X);
-			GraphData.Max.Y = FMath::Max(GraphData.Max.Y, CurrentSample.SampleValue.Y);
-		}
-	}
-	
-	if (bDirtyData)
-	{
-		const int32 ExtremeValueIndexes[] = { LeftSideOutsideIndex != INDEX_NONE ? LeftSideOutsideIndex : 0, RightSideOutsideIndex != INDEX_NONE ? RightSideOutsideIndex : ObjectItems.Num() - 1 };
-		for (int32 ObjectIndex = 0; ObjectIndex < 2; ++ObjectIndex)
-		{
-			const FVisualLogEntry* CurrentEntry = &(ObjectItems[ExtremeValueIndexes[ObjectIndex]].Entry);
-			const int32 SamplesNum = CurrentEntry->HistogramSamples.Num();
-			for (int32 SampleIndex = 0; SampleIndex < SamplesNum; ++SampleIndex)
-			{
-				FVisualLogHistogramSample CurrentSample = CurrentEntry->HistogramSamples[SampleIndex];
-
-				const FName CurrentCategory = CurrentSample.Category;
-				const FName CurrentGraphName = CurrentSample.GraphName;
-				const FName CurrentDataName = CurrentSample.DataName;
-
-				//FString GraphFilterName = CurrentSample.GraphName.ToString() + TEXT("_") + CurrentSample.DataName.ToString();
-				//const bool bIsValidByFilter = FCategoryFiltersManager::Get().MatchCategoryFilters(GraphFilterName, ELogVerbosity::All);
-				//if (bIsValidByFilter)
-				{
-					FGraphData *GraphData = CollectedGraphs.Find(CurrentSample.GraphName);
-					if (GraphData)
-					{
-						FGraphLineData *LineData = GraphData->GraphLines.Find(CurrentSample.DataName);
-						if (LineData)
-						{
-							if (ObjectIndex == 0)
-								LineData->LeftExtreme = CurrentSample.SampleValue;
-							else
-								LineData->RightExtreme = CurrentSample.SampleValue;
-						}
-					}
-				}
-			}
-		}
-	}
-
 	const float GoldenRatioConjugate = 0.618033988749895f;
 	if (CollectedGraphs.Num() > 0)
 	{
+		const FVisualLoggerTimeSliderArgs&  TimeSliderArgs = FLogVisualizer::Get().GetTimeSliderController()->GetTimeSliderArgs();
+		TRange<float> LocalViewRange = TimeSliderArgs.ViewRange.Get();
+		const float LocalViewRangeMin = LocalViewRange.GetLowerBoundValue();
+		const float LocalViewRangeMax = LocalViewRange.GetUpperBoundValue();
+		const float LocalSequenceLength = LocalViewRangeMax - LocalViewRangeMin;
+		const float WindowHalfWidth = LocalSequenceLength * TimeSliderArgs.CursorSize.Get() * 0.5f;
+		const FVector2D TimeStampWindow(SelectedEntry.TimeStamp - WindowHalfWidth, SelectedEntry.TimeStamp + WindowHalfWidth);
+
 		const FColor GraphsBackgroundColor = ULogVisualizerSettings::StaticClass()->GetDefaultObject<ULogVisualizerSettings>()->GraphsBackgroundColor;
 		const int NumberOfGraphs = CollectedGraphs.Num();
 		const int32 NumberOfColumns = FMath::CeilToInt(FMath::Sqrt(NumberOfGraphs));
diff --git a/Engine/Source/Developer/LogVisualizer/Private/VisualLoggerCanvasRenderer.h b/Engine/Source/Developer/LogVisualizer/Private/VisualLoggerCanvasRenderer.h
--- a/Engine/Source/Developer/LogVisualizer/Private/VisualLoggerCanvasRenderer.h
+++ b/Engine/Source/Developer/LogVisualizer/Private/VisualLoggerCanvasRenderer.h
@@ -3,14 +3,32 @@
 
 struct FVisualLoggerCanvasRenderer
 {
+	struct FGraphLineData
+	{
+		FName DataName;
+		FVector2D LeftExtreme, RightExtreme;
+		TArray<FVector2D> Samples;
+	};
+
+	struct FGraphData
+	{
+		FGraphData() : Min(FVector2D(FLT_MAX, FLT_MAX)), Max(FVector2D(FLT_MIN, FLT_MIN)) {}
+
+		FVector2D Min, Max;
+		TMap<FName, FGraphLineData> GraphLines;
+	};
+
+	TMap<FName, FGraphData>	CollectedGraphs;
+
 public:
-	FVisualLoggerCanvasRenderer() : bDirtyData(true)
-	{}
+	FVisualLoggerCanvasRenderer();
+	~FVisualLoggerCanvasRenderer();
 
 	void DrawOnCanvas(class UCanvas* Canvas, class APlayerController*);
-	void OnItemSelectionChanged(const struct FVisualLogEntry& EntryItem);
-	void ObjectSelectionChanged(TSharedPtr<class STimeline> TimeLine);
+	void OnItemSelectionChanged(const FVisualLoggerDBRow& ChangedRow, int32 SelectedItemIndex);
+	void ObjectSelectionChanged(const TArray<FName>& RowNames);
 	void DirtyCachedData() { bDirtyData = true; }
+	void ResetData();
 
 protected:
 	void DrawHistogramGraphs(class UCanvas* Canvas, class APlayerController* );
@@ -19,8 +37,4 @@ private:
 	bool bDirtyData;
 	FVisualLogEntry SelectedEntry;
 	TMap<FString, TArray<FString> > UsedGraphCategories;
-	TWeakPtr<class STimeline> CurrentTimeLine;
-
-	TArray<FVisualLogDataBlock> CachedDataBlocks;
-	TArray<TArray<FVisualLogHistogramSample> > CachedHistogramSamples;
 };
diff --git a/Engine/Source/Developer/LogVisualizer/Private/VisualLoggerCommands.h b/Engine/Source/Developer/LogVisualizer/Private/VisualLoggerCommands.h
--- a/Engine/Source/Developer/LogVisualizer/Private/VisualLoggerCommands.h
+++ b/Engine/Source/Developer/LogVisualizer/Private/VisualLoggerCommands.h
@@ -31,9 +31,6 @@ public:
 		UI_COMMAND(ToggleGraphs, "ToggleGraphs", "Toggle graphs visualization on/off", EUserInterfaceActionType::ToggleButton, FInputChord());
 		UI_COMMAND(ResetData, "Clear", "Clear all data", EUserInterfaceActionType::Button, FInputChord());
 		UI_COMMAND(HideEmpty, "HideEmpty", "Hide logs without text information", EUserInterfaceActionType::ToggleButton, FInputChord());
-
-		UI_COMMAND(MoveCursorLeft, "MoveCursorLeft", "Move cursor left on timeline", EUserInterfaceActionType::Button, FInputChord(EKeys::Left));
-		UI_COMMAND(MoveCursorRight, "MoveCursorLeft", "Move cursor right on timeline", EUserInterfaceActionType::Button, FInputChord(EKeys::Right));
 	}
 
 public:
@@ -49,8 +46,6 @@ public:
 	TSharedPtr<FUICommandInfo> ToggleGraphs;
 	TSharedPtr<FUICommandInfo> ResetData;
 	TSharedPtr<FUICommandInfo> HideEmpty;
-	TSharedPtr<FUICommandInfo> MoveCursorLeft;
-	TSharedPtr<FUICommandInfo> MoveCursorRight;
 };
 
 
diff --git a/Engine/Source/Developer/LogVisualizer/Private/VisualLoggerDatabase.cpp b/Engine/Source/Developer/LogVisualizer/Private/VisualLoggerDatabase.cpp
new file mode 100644
--- /dev/null
+++ b/Engine/Source/Developer/LogVisualizer/Private/VisualLoggerDatabase.cpp
@@ -0,0 +1,463 @@
+// Copyright 1998-2015 Epic Games, Inc. All Rights Reserved.
+
+#include "LogVisualizer.h"
+#include "Engine.h"
+#include "VisualLoggerDatabase.h"
+
+
+TSharedPtr< struct FVisualLoggerDatabase > FVisualLoggerDatabase::StaticInstance = nullptr;
+TSharedPtr< struct FVisualLoggerGraphsDatabase > FVisualLoggerGraphsDatabase::StaticInstance = nullptr;
+
+void FVisualLoggerDBRow::AddItem(const FVisualLogDevice::FVisualLogEntryItem& NewItem) 
+{ 
+	const int32 ItemIndex = Items.Add(NewItem);
+	DBEvents.OnNewItem.Broadcast(*this, ItemIndex);
+}
+
+void FVisualLoggerDBRow::MoveTo(int32 Index) 
+{ 
+	const int32 OldItemIndex = CurrentItemIndex;
+	CurrentItemIndex = Items.IsValidIndex(Index) ? Index : INDEX_NONE; 
+	if (OldItemIndex != CurrentItemIndex)
+	{
+		DBEvents.OnItemSelectionChanged.Broadcast(*this, CurrentItemIndex);
+	}
+}
+
+const FVisualLogDevice::FVisualLogEntryItem& FVisualLoggerDBRow::GetCurrentItem() const 
+{ 
+	check(Items.IsValidIndex(CurrentItemIndex)); 
+	return Items[CurrentItemIndex]; 
+}
+
+void FVisualLoggerDBRow::SetItemVisibility(int32 ItemIndex, bool IsVisible)
+{ 
+	if (IsVisible)
+	{
+		HiddenItems.RemoveSingleSwap(ItemIndex);
+	}
+	else
+	{
+		HiddenItems.AddUnique(ItemIndex);
+	}
+}
+
+int32 FVisualLoggerDBRow::GetClosestItem(float Time) const
+{
+	int32 BestItemIndex = INDEX_NONE;
+	float BestDistance = MAX_FLT;
+	for (int32 Index = 0; Index < Items.Num(); Index++)
+	{
+		auto& CurrentEntryItem = Items[Index];
+
+		if (IsItemVisible(Index) == false)
+		{
+			continue;
+		}
+		TArray<FVisualLoggerCategoryVerbosityPair> OutCategories;
+		const float CurrentDist = FMath::Abs(Time - CurrentEntryItem.Entry.TimeStamp);
+		if (CurrentDist < BestDistance)
+		{
+			BestDistance = CurrentDist;
+			BestItemIndex = Index;
+		}
+	}
+
+	const float CurrentDist = Items.IsValidIndex(CurrentItemIndex) && IsItemVisible(CurrentItemIndex) ? FMath::Abs(Time - Items[CurrentItemIndex].Entry.TimeStamp) : MAX_FLT;
+	if (BestItemIndex != INDEX_NONE && CurrentDist > BestDistance)
+	{
+		return BestItemIndex;
+	}
+
+	return CurrentItemIndex;
+}
+
+int32 FVisualLoggerDBRow::GetClosestItem(float Time, float ScrubTime) const
+{
+	int32 BestItemIndex = INDEX_NONE;
+	float BestDistance = MAX_FLT;
+	for (int32 Index = 0; Index < Items.Num(); Index++)
+	{
+		auto& CurrentEntryItem = Items[Index];
+
+		if (CurrentEntryItem.Entry.TimeStamp > ScrubTime)
+		{
+			break;
+		}
+
+		if (IsItemVisible(Index) == false)
+		{
+			continue;
+		}
+		TArray<FVisualLoggerCategoryVerbosityPair> OutCategories;
+		const float CurrentDist = FMath::Abs(CurrentEntryItem.Entry.TimeStamp - Time);
+		if (CurrentDist < BestDistance && CurrentEntryItem.Entry.TimeStamp <= Time)
+		{
+			BestDistance = CurrentDist;
+			BestItemIndex = Index;
+		}
+	}
+
+
+	const float CurrentDist = Items.IsValidIndex(CurrentItemIndex) ? FMath::Abs(Items[CurrentItemIndex].Entry.TimeStamp - Time) : MAX_FLT;
+
+	if (BestItemIndex != INDEX_NONE)
+	{
+		return BestItemIndex;
+	}
+
+	return CurrentItemIndex;
+}
+
+FVisualLoggerDatabase& FVisualLoggerDatabase::Get()
+{
+	return *StaticInstance;
+}
+
+void FVisualLoggerDatabase::Initialize()
+{
+	StaticInstance = MakeShareable(new FVisualLoggerDatabase);
+	FVisualLoggerGraphsDatabase::Initialize();
+}
+
+void FVisualLoggerDatabase::Shutdown()
+{
+	FVisualLoggerGraphsDatabase::Shutdown();
+	StaticInstance.Reset();
+}
+
+void FVisualLoggerDatabase::Reset()
+{
+	Rows.Reset();
+	RowNameToIndex.Reset();
+	SelectedRows.Reset();
+	HiddenRows.Reset();
+	FVisualLoggerGraphsDatabase::Get().Reset();
+}
+
+void FVisualLoggerDatabase::AddItem(const FVisualLogDevice::FVisualLogEntryItem& NewItem)
+{
+	const bool bCreateNew = RowNameToIndex.Contains(NewItem.OwnerName) == false;
+	int32 RowIndex = INDEX_NONE;
+	FVisualLoggerDBRow* CurrentRow = nullptr;
+
+	if (bCreateNew)
+	{
+		RowIndex = Rows.Add(FVisualLoggerDBRow(DBEvents, NewItem.OwnerName, NewItem.OwnerClassName));
+		RowNameToIndex.Add(NewItem.OwnerName, RowIndex);
+		CurrentRow = &Rows[RowIndex];
+		DBEvents.OnNewRow.Broadcast(*CurrentRow);
+	}
+	else
+	{
+		RowIndex = RowNameToIndex[NewItem.OwnerName];
+		CurrentRow = &Rows[RowIndex];
+	}
+	check(CurrentRow);
+
+	CurrentRow->AddItem(NewItem);
+
+	FVisualLoggerGraphsDatabase::Get().AddItem(NewItem);
+}
+
+FVisualLoggerDBRow& FVisualLoggerDatabase::GetRowByName(FName InName)
+{
+	const bool bContainsRow = RowNameToIndex.Contains(InName);
+	check(bContainsRow);
+
+	return Rows[RowNameToIndex[InName]];
+}
+
+void FVisualLoggerDatabase::SelectRow(FName InName, bool bDeselectOtherNodes)
+{
+	const bool bAlreadySelected = SelectedRows.Find(InName) != INDEX_NONE;
+	if (bAlreadySelected && (!bDeselectOtherNodes || SelectedRows.Num() == 1))
+	{
+		return;
+	}
+
+	if (bDeselectOtherNodes)
+	{
+		for (auto CurrentName : SelectedRows)
+		{
+			if (CurrentName != InName)
+			{
+				FVisualLoggerDBRow& DBRow = GetRowByName(CurrentName);
+				DBRow.MoveTo(INDEX_NONE);
+			}
+		}
+		SelectedRows.Reset();
+		if (bAlreadySelected)
+		{
+			SelectedRows.AddUnique(InName);
+			DBEvents.OnRowSelectionChanged.Broadcast(SelectedRows);
+			return;
+		}
+	}
+
+	if (!bAlreadySelected)
+	{
+		SelectedRows.AddUnique(InName);
+		DBEvents.OnRowSelectionChanged.Broadcast(SelectedRows);
+	}
+}
+
+void FVisualLoggerDatabase::DeselectRow(FName InName)
+{
+	const bool bSelected = SelectedRows.Find(InName) != INDEX_NONE;
+	if (bSelected)
+	{
+		FVisualLoggerDBRow& DBRow = GetRowByName(InName);
+		DBRow.MoveTo(INDEX_NONE);
+
+		SelectedRows.RemoveSingleSwap(InName, false);
+		DBEvents.OnRowSelectionChanged.Broadcast(SelectedRows);
+	}
+}
+
+const TArray<FName>& FVisualLoggerDatabase::GetSelectedRows() const
+{
+	return SelectedRows;
+}
+
+void FVisualLoggerDatabase::SetRowVisibility(FName RowName, bool SetAsVisible)
+{
+	const bool IsVisible = HiddenRows.Find(RowName) == INDEX_NONE;
+	if (IsVisible != SetAsVisible)
+	{
+		if (SetAsVisible)
+		{
+			HiddenRows.RemoveSingleSwap(RowName);
+		}
+		else
+		{
+			HiddenRows.AddUnique(RowName);
+		}
+
+		DBEvents.OnRowChangedVisibility.Broadcast(RowName);
+	}
+}
+
+void FVisualLoggerDatabase::RemoveRow(FName RowName)
+{
+	if (SelectedRows.Find(RowName) != INDEX_NONE)
+	{
+		SelectedRows.RemoveSwap(RowName);
+	}
+	if (HiddenRows.Find(RowName) != INDEX_NONE)
+	{
+		HiddenRows.RemoveSwap(RowName);
+	}
+	if (RowNameToIndex.Contains(RowName))
+	{
+		const int32 RemovedIndex = RowNameToIndex[RowName];
+		Rows.RemoveAtSwap(RemovedIndex, 1, false);
+		if (Rows.IsValidIndex(RemovedIndex))
+		{
+			RowNameToIndex[Rows[RemovedIndex].GetOwnerName()] = RemovedIndex;
+		}
+	}
+}
+
+/************************************************************************/
+/*                                                                      */
+/************************************************************************/
+
+bool FVisualLoggerGraph::IsDataVisible(FName DataName) const
+{
+	return HiddenGraphs.Contains(DataName) == false;
+}
+
+void FVisualLoggerGraph::SetDataVisibility(FName DataName, bool IsVisible)
+{
+	if (IsVisible)
+	{
+		HiddenGraphs.Remove(DataName);
+	}
+	else
+	{
+		HiddenGraphs.AddUnique(DataName);
+	}
+}
+
+FVisualLoggerGraphData& FVisualLoggerGraph::FindOrAddDataByName(FName DataName)
+{
+	if (DataNameToIndex.Contains(DataName))
+	{
+		return DataGraphs[DataNameToIndex[DataName]];
+	}
+
+	const int32 Index = DataGraphs.Add(FVisualLoggerGraphData(DataName));
+	DataNameToIndex.Add(DataName, Index);
+	FVisualLoggerDatabase::Get().GetEvents().OnGraphDataNameAddedEvent.Broadcast(OwnerName, GraphName, DataName);
+
+	return DataGraphs[Index];
+}
+
+/************************************************************************/
+/*                                                                      */
+/************************************************************************/
+
+FVisualLoggerGraphsDatabase& FVisualLoggerGraphsDatabase::Get()
+{
+	return *StaticInstance;
+}
+
+void FVisualLoggerGraphsDatabase::Initialize()
+{
+	StaticInstance = MakeShareable(new FVisualLoggerGraphsDatabase);
+}
+
+void FVisualLoggerGraphsDatabase::Shutdown()
+{
+	StaticInstance.Reset();
+}
+
+void FVisualLoggerGraphsDatabase::Reset()
+{
+	OwnerNameToGraphs.Reset();
+	HiddenGraphs.Reset();
+}
+
+void FVisualLoggerGraphsDatabase::AddItem(const FVisualLogDevice::FVisualLogEntryItem& NewItem)
+{
+	FVisualLoggerGraphHelper* GraphHelperPtr = OwnerNameToGraphs.Find(NewItem.OwnerName);
+	if (!GraphHelperPtr)
+	{
+		GraphHelperPtr = &OwnerNameToGraphs.Add(NewItem.OwnerName);
+	}
+	FVisualLoggerGraphHelper& GraphHelper = *GraphHelperPtr;
+
+	for (const FVisualLogHistogramSample& HistogramSample : NewItem.Entry.HistogramSamples)
+	{
+		int32 GraphIndex = INDEX_NONE;
+		if (GraphHelper.GraphNameToIndex.Contains(HistogramSample.GraphName) == false)
+		{
+			FVisualLoggerGraph Graph(NewItem.OwnerName);
+			Graph.SetGraphName(HistogramSample.GraphName);
+			GraphIndex = GraphHelper.AllGraphs.Add(Graph);
+			GraphHelper.GraphNameToIndex.Add(HistogramSample.GraphName, GraphIndex);
+
+			FVisualLoggerDatabase::Get().GetEvents().OnGraphAddedEvent.Broadcast(NewItem.OwnerName, HistogramSample.GraphName);
+		}
+		else
+		{
+			GraphIndex = GraphHelper.GraphNameToIndex[HistogramSample.GraphName];
+		}
+
+		check(GraphIndex != INDEX_NONE);
+		FVisualLoggerGraphData& GraphData = GraphHelper.AllGraphs[GraphIndex].FindOrAddDataByName(HistogramSample.DataName);
+		GraphData.Samples.Add(HistogramSample.SampleValue);
+		GraphData.TimeStamps.Add(NewItem.Entry.TimeStamp);
+	}
+}
+
+bool FVisualLoggerGraphsDatabase::IsGraphVisible(FName OwnerName, FName GraphName)
+{
+	return HiddenGraphs.Find(*(OwnerName.ToString() + TEXT("$") + GraphName.ToString())) == INDEX_NONE;
+}
+
+void FVisualLoggerGraphsDatabase::SetGraphVisibility(FName OwnerName, FName GraphName, bool SetAsVisible)
+{
+	const FName FullName = *(OwnerName.ToString() + TEXT("$") + GraphName.ToString());
+	const bool IsVisible = HiddenGraphs.Find(FullName) == INDEX_NONE;
+	if (IsVisible != SetAsVisible)
+	{
+		if (SetAsVisible)
+		{
+			HiddenGraphs.RemoveSingleSwap(FullName);
+		}
+		else
+		{
+			HiddenGraphs.Add(FullName);
+		}
+
+		FVisualLoggerDatabase::Get().GetEvents().OnGraphChangedVisibilityEvent.Broadcast(GraphName);
+	}
+}
+
+bool FVisualLoggerGraphsDatabase::ContainsGraphByName(FName OwnerName, FName GraphName)
+{
+	const FVisualLoggerGraphHelper& GraphHelper = OwnerNameToGraphs.FindOrAdd(OwnerName);
+	return GraphHelper.GraphNameToIndex.Contains(GraphName);
+}
+
+FVisualLoggerGraph& FVisualLoggerGraphsDatabase::GetGraphByName(FName OwnerName, FName GraphName)
+{
+	FVisualLoggerGraphHelper& GraphHelper = OwnerNameToGraphs.FindOrAdd(OwnerName);
+	const bool bContainsRow = GraphHelper.GraphNameToIndex.Contains(GraphName);
+	check(bContainsRow);
+
+	const int32 GraphIndex = GraphHelper.GraphNameToIndex[GraphName];
+	return GraphHelper.AllGraphs[GraphIndex];
+}
+
+const TArray<FVisualLoggerGraph>& FVisualLoggerGraphsDatabase::GetGraphsByOwnerName(FName OwnerName)
+{
+	const FVisualLoggerGraphHelper& GraphHelper = OwnerNameToGraphs.FindOrAdd(OwnerName);
+	return GraphHelper.AllGraphs;
+}
+
+/************************************************************************/
+/* FVisualLoggerEditorInterface                                         */
+/************************************************************************/
+const FName& FVisualLoggerEditorInterface::GetRowClassName(FName RowName) const 
+{ 
+	return FVisualLoggerDatabase::Get().GetRowByName(RowName).GetOwnerClassName(); 
+}
+
+int32 FVisualLoggerEditorInterface::GetSelectedItemIndex(FName RowName) const 
+{ 
+	return FVisualLoggerDatabase::Get().GetRowByName(RowName).GetCurrentItemIndex(); 
+}
+
+const TArray<FVisualLogDevice::FVisualLogEntryItem>& FVisualLoggerEditorInterface::GetRowItems(FName RowName) 
+{ 
+	return FVisualLoggerDatabase::Get().GetRowByName(RowName).GetItems(); 
+}
+
+const FVisualLogDevice::FVisualLogEntryItem& FVisualLoggerEditorInterface::GetSelectedItem(FName RowName) const 
+{ 
+	return FVisualLoggerDatabase::Get().GetRowByName(RowName).GetCurrentItem(); 
+}
+
+const TArray<FName>& FVisualLoggerEditorInterface::GetSelectedRows() const 
+{ 
+	return FVisualLoggerDatabase::Get().GetSelectedRows(); 
+}
+
+bool FVisualLoggerEditorInterface::IsRowVisible(FName RowName) const
+{ 
+	return FVisualLoggerDatabase::Get().IsRowVisible(RowName); 
+}
+
+UWorld* FVisualLoggerEditorInterface::GetWorld() const
+{ 
+	return FLogVisualizer::Get().GetWorld(); 
+}
+
+bool FVisualLoggerEditorInterface::IsItemVisible(FName RowName, int32 ItemIndex) const
+{
+	return FVisualLoggerDatabase::Get().GetRowByName(RowName).IsItemVisible(ItemIndex);
+}
+
+AActor* FVisualLoggerEditorInterface::GetHelperActor() const
+{
+	UWorld* World = GetWorld();
+	if (World == nullptr)
+	{
+		return nullptr;
+	}
+
+	for (TActorIterator<AVisualLoggerRenderingActor> It(World); It; ++It)
+	{
+		return *It;
+	}
+
+	return nullptr;
+}
+
+bool FVisualLoggerEditorInterface::MatchCategoryFilters(const FString& String, ELogVerbosity::Type Verbosity)
+{
+	return FVisualLoggerFilters::Get().MatchCategoryFilters(String, Verbosity);
+}
diff --git a/Engine/Source/Developer/LogVisualizer/Private/VisualLoggerDatabase.h b/Engine/Source/Developer/LogVisualizer/Private/VisualLoggerDatabase.h
new file mode 100644
--- /dev/null
+++ b/Engine/Source/Developer/LogVisualizer/Private/VisualLoggerDatabase.h
@@ -0,0 +1,207 @@
+// Copyright 1998-2015 Epic Games, Inc. All Rights Reserved.
+#pragma once
+#include "ObjectBase.h"
+#include "SDockTab.h"
+#include "VisualLogger/VisualLogger.h"
+
+struct FVisualLoggerDBRow;
+
+struct FVisualLoggerDBEvents
+{
+	DECLARE_MULTICAST_DELEGATE_TwoParams(FItemSelectionChangedEvent, const FVisualLoggerDBRow&, int32);
+	DECLARE_MULTICAST_DELEGATE_TwoParams(FNewItemEvent, const FVisualLoggerDBRow&, int32);
+	DECLARE_MULTICAST_DELEGATE_OneParam(FNewRowEvent, const FVisualLoggerDBRow&);
+	DECLARE_MULTICAST_DELEGATE_OneParam(FOnRowSelectionChangedEvent, const TArray<FName>&);
+	DECLARE_MULTICAST_DELEGATE_OneParam(FOnRowChangedVisibilityEvent, const FName&);
+	DECLARE_MULTICAST_DELEGATE_OneParam(FOnRowRemovedEvent, const FName&);
+
+	DECLARE_MULTICAST_DELEGATE_OneParam(FOnGraphChangedVisibilityEvent, const FName&);
+	DECLARE_MULTICAST_DELEGATE_TwoParams(FOnGraphAddedEvent, const FName&, const FName&);
+	DECLARE_MULTICAST_DELEGATE_ThreeParams(FOnGraphDataNameAddedEvent, const FName&, const FName&, const FName&);
+
+	FNewItemEvent OnNewItem;
+	FItemSelectionChangedEvent OnItemSelectionChanged;
+	FNewRowEvent OnNewRow;
+	FOnRowSelectionChangedEvent OnRowSelectionChanged;
+	FOnRowChangedVisibilityEvent OnRowChangedVisibility;
+	FOnRowRemovedEvent OnRowRemoved;
+
+	FOnGraphChangedVisibilityEvent OnGraphChangedVisibilityEvent;
+	FOnGraphAddedEvent OnGraphAddedEvent;
+	FOnGraphDataNameAddedEvent OnGraphDataNameAddedEvent;
+};
+
+struct FVisualLoggerDBRow
+{
+public:
+	FVisualLoggerDBRow(FVisualLoggerDBEvents& InEvents, const FName& InOwnerName, const FName& InOwnerClassName) : DBEvents(InEvents), OwnerName(InOwnerName), OwnerClassName(InOwnerClassName), CurrentItemIndex(INDEX_NONE) {}
+
+	const FName& GetOwnerName() const { return OwnerName; }
+	const FName& GetOwnerClassName() const { return OwnerClassName; }
+
+	void AddItem(const FVisualLogDevice::FVisualLogEntryItem& NewItem);
+	const TArray<FVisualLogDevice::FVisualLogEntryItem>& GetItems() const { return Items; }
+
+	void MoveTo(int32 Index);
+	const FVisualLogDevice::FVisualLogEntryItem& GetCurrentItem() const;
+	int32 GetCurrentItemIndex() const { return CurrentItemIndex; }
+
+	void SetItemVisibility(int32 ItemIndex, bool IsVisible);
+	bool IsItemVisible(int32 ItemIndex) const { return HiddenItems.Find(ItemIndex) == INDEX_NONE; }
+	int32 GetNumberOfHiddenItems() const { return HiddenItems.Num(); }
+
+	int32 GetClosestItem(float Time) const;
+	int32 GetClosestItem(float Time, float ScrubTime) const;
+
+protected:
+	FVisualLoggerDBEvents& DBEvents;
+	FName OwnerName;
+	FName OwnerClassName;
+	int32 CurrentItemIndex;
+
+	TArray<FVisualLogDevice::FVisualLogEntryItem>	Items;
+	TArray<int32> HiddenItems;
+};
+
+struct FVisualLoggerDatabase
+{
+	typedef TArray<FVisualLoggerDBRow>::TConstIterator FConstRowIterator;
+	typedef TArray<FVisualLoggerDBRow>::TIterator FRowIterator;
+
+	static FVisualLoggerDatabase& Get();
+	static void Initialize();
+	static void Shutdown();
+
+	void Reset();
+	FVisualLoggerDBEvents& GetEvents() { return DBEvents; }
+
+	int32 NumberOfRows() { return Rows.Num(); }
+	void AddItem(const FVisualLogDevice::FVisualLogEntryItem& NewItem);
+	FConstRowIterator GetConstRowIterator() const { return Rows.CreateConstIterator(); }
+	FRowIterator GetRowIterator() { return Rows.CreateIterator(); }
+	
+	FVisualLoggerDBRow& GetRowByName(FName InName);
+
+	void SelectRow(FName InName, bool bDeselectOtherNodes = false);
+	void DeselectRow(FName InName);
+	bool IsRowSelected(FName InName) const { return SelectedRows.Find(InName) != INDEX_NONE; }
+	const TArray<FName>& GetSelectedRows() const;
+
+	bool IsRowVisible(FName RowName) const { return HiddenRows.Find(RowName) == INDEX_NONE; }
+	void SetRowVisibility(FName RowName, bool IsVisible);
+
+	void RemoveRow(FName RowName);
+
+protected:
+	TArray<FVisualLoggerDBRow>	Rows;
+	TMap<FName, int32>	RowNameToIndex;
+	TArray<FName> SelectedRows;
+	TArray<FName> HiddenRows;
+
+	FVisualLoggerDBEvents DBEvents;
+
+private:
+	static TSharedPtr< struct FVisualLoggerDatabase > StaticInstance;
+
+};
+
+struct FVisualLoggerGraphData
+{
+	FName DataName;
+	TArray<FVector2D> Samples;
+	TArray<float> TimeStamps;
+
+	FVisualLoggerGraphData(FName InDataName) : DataName(InDataName) {}
+};
+
+struct FVisualLoggerGraph
+{
+	typedef TArray<FVisualLoggerGraphData>::TConstIterator FConstDataIterator;
+	typedef TArray<FVisualLoggerGraphData>::TIterator FDataIterator;
+
+	FVisualLoggerGraph(FName InOwnerName) : OwnerName(InOwnerName) {}
+
+	FName GetOwnerName() const { return OwnerName; }
+	FName GetGraphName() const { return GraphName; }
+	void SetGraphName(FName InGraphName) { GraphName = InGraphName; }
+
+	bool IsDataVisible(FName DataName) const;
+	void SetDataVisibility(FName DataName, bool IsVisible);
+
+	bool ContainsDataByName(FName DataName) const { return DataNameToIndex.Contains(DataName); };
+	FVisualLoggerGraphData& FindOrAddDataByName(FName DataName);
+
+	FConstDataIterator GetConstDataIterator() const { return DataGraphs.CreateConstIterator(); }
+	FDataIterator GetDataIterator() { return DataGraphs.CreateIterator(); }
+
+protected:
+	FName OwnerName;
+	FName GraphName;
+	TArray<FVisualLoggerGraphData> DataGraphs;
+	TMap<FName, int32>	DataNameToIndex;
+	TArray<FName> HiddenGraphs;
+};
+
+struct FVisualLoggerGraphsDatabase //histogram graphs database as separate structure to optimize access and category filters
+{
+protected:
+	struct FVisualLoggerGraphHelper
+	{
+		TArray<FVisualLoggerGraph> AllGraphs;
+		TMap<FName, int32> GraphNameToIndex;
+	};
+
+public:
+	typedef TArray<FVisualLoggerGraph>::TConstIterator FConstGraphIterator;
+	typedef TArray<FVisualLoggerGraph>::TIterator FGraphIterator;
+	
+	typedef TMap<FName, FVisualLoggerGraphHelper>::TConstIterator FConstOwnersIterator;
+	typedef TMap<FName, FVisualLoggerGraphHelper>::TIterator FOwnersIterator;
+
+	static FVisualLoggerGraphsDatabase& Get();
+	static void Initialize();
+	static void Shutdown();
+	void Reset();
+
+	void AddItem(const FVisualLogDevice::FVisualLogEntryItem& NewItem);
+
+	bool IsGraphVisible(FName OwnerName, FName GraphName);
+	void SetGraphVisibility(FName OwnerName, FName GraphName, bool IsVisible);
+
+	bool ContainsGraphByName(FName OwnerName, FName GraphName);
+	FVisualLoggerGraph& GetGraphByName(FName OwnerName, FName GraphName);
+	bool ContainsHistogramGraphs() const { return OwnerNameToGraphs.Num() > 0; }
+
+	const TArray<FVisualLoggerGraph>& GetGraphsByOwnerName(FName OwnerName);
+
+	FConstGraphIterator GetConstGraphsIterator(FName OwnerName) { return OwnerNameToGraphs.FindOrAdd(OwnerName).AllGraphs.CreateConstIterator(); }
+	FGraphIterator GetGraphsIterator(FName OwnerName) { return OwnerNameToGraphs.FindOrAdd(OwnerName).AllGraphs.CreateIterator(); }
+
+	FConstOwnersIterator GetConstOwnersIterator() { return OwnerNameToGraphs.CreateConstIterator(); }
+	FOwnersIterator GetOwnersIterator() { return OwnerNameToGraphs.CreateIterator(); }
+
+protected:
+	TMap<FName, FVisualLoggerGraphHelper> OwnerNameToGraphs;
+	TArray<FName> HiddenGraphs; //encoded as "OwnerName$GraphName" for simplicity.
+
+private:
+	static TSharedPtr< struct FVisualLoggerGraphsDatabase > StaticInstance;
+};
+
+struct FVisualLoggerEditorInterface : public IVisualLoggerEditorInterface
+{
+	static IVisualLoggerEditorInterface* Get() { static FVisualLoggerEditorInterface EditorInterface; return &EditorInterface; }
+
+	const FName& GetRowClassName(FName RowName) const override;
+	int32 GetSelectedItemIndex(FName RowName) const override;
+	const TArray<FVisualLogDevice::FVisualLogEntryItem>& GetRowItems(FName RowName) override;
+	const FVisualLogDevice::FVisualLogEntryItem& GetSelectedItem(FName RowName) const override;
+
+	const TArray<FName>& GetSelectedRows() const override;
+	bool IsRowVisible(FName RowName) const override;
+	bool IsItemVisible(FName RowName, int32 ItemIndex) const override;
+	UWorld* GetWorld() const override;
+	AActor* GetHelperActor() const override;
+
+	bool MatchCategoryFilters(const FString& String, ELogVerbosity::Type Verbosity = ELogVerbosity::All) override;
+};
diff --git a/Engine/Source/Developer/LogVisualizer/Private/VisualLoggerRenderingActor.cpp b/Engine/Source/Developer/LogVisualizer/Private/VisualLoggerRenderingActor.cpp
--- a/Engine/Source/Developer/LogVisualizer/Private/VisualLoggerRenderingActor.cpp
+++ b/Engine/Source/Developer/LogVisualizer/Private/VisualLoggerRenderingActor.cpp
@@ -59,26 +59,18 @@ FPrimitiveSceneProxy* UVisualLoggerRenderingComponent::CreateSceneProxy()
 	ULogVisualizerSettings *Settings = ULogVisualizerSettings::StaticClass()->GetDefaultObject<ULogVisualizerSettings>();
 	FVisualLoggerSceneProxy *VLogSceneProxy = new FVisualLoggerSceneProxy(this);
 	VLogSceneProxy->SolidMeshMaterial = Settings->GetDebugMeshMaterial();
-	VLogSceneProxy->Spheres = RenderingActor->PrimaryDebugShapes.Points;
-	VLogSceneProxy->Lines = RenderingActor->PrimaryDebugShapes.Lines;
-	VLogSceneProxy->Cones = RenderingActor->PrimaryDebugShapes.Cones;
-	VLogSceneProxy->Boxes = RenderingActor->PrimaryDebugShapes.Boxes;
-	VLogSceneProxy->Meshes = RenderingActor->PrimaryDebugShapes.Meshes;
-	VLogSceneProxy->Cones = RenderingActor->PrimaryDebugShapes.Cones;
-	VLogSceneProxy->Texts = RenderingActor->PrimaryDebugShapes.Texts;
-	VLogSceneProxy->Cylinders = RenderingActor->PrimaryDebugShapes.Cylinders;
-	VLogSceneProxy->Capsles = RenderingActor->PrimaryDebugShapes.Capsles;
 
+	for (auto& CurrentShapes : RenderingActor->DebugShapesPerRow)
 	{
-		VLogSceneProxy->Spheres.Append(RenderingActor->SecondaryDebugShapes.Points);
-		VLogSceneProxy->Lines.Append(RenderingActor->SecondaryDebugShapes.Lines);
-		VLogSceneProxy->Cones.Append(RenderingActor->SecondaryDebugShapes.Cones);
-		VLogSceneProxy->Boxes.Append(RenderingActor->SecondaryDebugShapes.Boxes);
-		VLogSceneProxy->Meshes.Append(RenderingActor->SecondaryDebugShapes.Meshes);
-		VLogSceneProxy->Cones.Append(RenderingActor->SecondaryDebugShapes.Cones);
-		VLogSceneProxy->Texts.Append(RenderingActor->SecondaryDebugShapes.Texts);
-		VLogSceneProxy->Cylinders.Append(RenderingActor->SecondaryDebugShapes.Cylinders);
-		VLogSceneProxy->Capsles.Append(RenderingActor->SecondaryDebugShapes.Capsles);
+		VLogSceneProxy->Spheres.Append(CurrentShapes.Value.Points);
+		VLogSceneProxy->Lines.Append(CurrentShapes.Value.Lines);
+		VLogSceneProxy->Cones.Append(CurrentShapes.Value.Cones);
+		VLogSceneProxy->Boxes.Append(CurrentShapes.Value.Boxes);
+		VLogSceneProxy->Meshes.Append(CurrentShapes.Value.Meshes);
+		VLogSceneProxy->Cones.Append(CurrentShapes.Value.Cones);
+		VLogSceneProxy->Texts.Append(CurrentShapes.Value.Texts);
+		VLogSceneProxy->Cylinders.Append(CurrentShapes.Value.Cylinders);
+		VLogSceneProxy->Capsles.Append(CurrentShapes.Value.Capsles);
 	}
 
 	{
@@ -133,7 +125,126 @@ void UVisualLoggerRenderingComponent::DestroyRenderState_Concurrent()
 AVisualLoggerRenderingActor::AVisualLoggerRenderingActor(const FObjectInitializer& ObjectInitializer)
 	: Super(ObjectInitializer)
 {
+	USceneComponent* SceneComponent = CreateDefaultSubobject<USceneComponent>(TEXT("SceneComp"));
+	RootComponent = SceneComponent;
+
 	RenderingComponent = CreateDefaultSubobject<UVisualLoggerRenderingComponent>(TEXT("RenderingComponent"));
+
+	if (HasAnyFlags(RF_ClassDefaultObject) == false)
+	{
+		FVisualLoggerDatabase::Get().GetEvents().OnItemSelectionChanged.AddUObject(this, &AVisualLoggerRenderingActor::OnItemSelectionChanged);
+		FVisualLoggerDatabase::Get().GetEvents().OnRowSelectionChanged.AddUObject(this, &AVisualLoggerRenderingActor::ObjectSelectionChanged);
+
+		FLogVisualizer::Get().GetEvents().OnFiltersChanged.AddUObject(this, &AVisualLoggerRenderingActor::OnFiltersChanged);
+	}
+}
+
+AVisualLoggerRenderingActor::~AVisualLoggerRenderingActor()
+{
+	if (HasAnyFlags(RF_ClassDefaultObject) == false)
+	{
+		FVisualLoggerDatabase::Get().GetEvents().OnItemSelectionChanged.RemoveAll(this);
+		FVisualLoggerDatabase::Get().GetEvents().OnRowSelectionChanged.RemoveAll(this);
+		FLogVisualizer::Get().GetEvents().OnFiltersChanged.RemoveAll(this);
+	}
+}
+
+void AVisualLoggerRenderingActor::ObjectSelectionChanged(const TArray<FName>& Selection)
+{
+	//PrimaryDebugShapes.Reset();
+
+	if (Selection.Num() > 0)
+	{
+		for (auto CurrentName : Selection)
+		{
+			if (DebugShapesPerRow.Contains(CurrentName) == false)
+			{
+				DebugShapesPerRow.Add(CurrentName);
+				FVisualLoggerDBRow &DBRow = FVisualLoggerDatabase::Get().GetRowByName(CurrentName);
+				FTimelineDebugShapes& ShapesCache = DebugShapesPerRow[CurrentName];
+				for (const auto &CurrentEntry : DBRow.GetItems())
+				{
+					if (CurrentEntry.Entry.Location != FVector::ZeroVector)
+					{
+						ShapesCache.LogEntriesPath.Add(CurrentEntry.Entry.Location);
+					}
+				}
+			}
+		}
+
+		for (auto Element : DebugShapesPerRow)
+		{
+			if (Selection.Find(Element.Key) == INDEX_NONE)
+			{
+				DebugShapesPerRow.Remove(Element.Key);
+			}
+		}
+	}
+	else
+	{
+		DebugShapesPerRow.Reset();
+	}
+	CachedRowSelection = Selection;
+	MarkComponentsRenderStateDirty();
+}
+
+void AVisualLoggerRenderingActor::OnItemSelectionChanged(const FVisualLoggerDBRow& DBRow, int32 ItemIndex)
+{
+	const FName RowName = DBRow.GetOwnerName();
+
+	const TMap<FName, FVisualLogExtensionInterface*>& AllExtensions = FVisualLogger::Get().GetAllExtensions();
+	for (auto& Extension : AllExtensions)
+	{
+		Extension.Value->DrawData(FVisualLoggerEditorInterface::Get(), NULL);
+	}
+
+	if (DebugShapesPerRow.Contains(RowName) == false)
+	{
+		return;
+	}
+
+	if (FVisualLoggerDatabase::Get().IsRowVisible(RowName) == false || DBRow.GetItems().IsValidIndex(ItemIndex) == false)
+	{
+		return;
+	}
+
+	FTimelineDebugShapes& ShapesCache = DebugShapesPerRow[RowName];
+	ShapesCache.Reset();
+	const TArray<FVisualLogDevice::FVisualLogEntryItem>& Entries = DBRow.GetItems();
+	const int32 CurrentItemIndex = DBRow.GetCurrentItemIndex();
+	GetDebugShapes(Entries[CurrentItemIndex], ShapesCache);
+
+	MarkComponentsRenderStateDirty();
+}
+
+void AVisualLoggerRenderingActor::ResetRendering()
+{
+	CachedRowSelection.Reset();
+	DebugShapesPerRow.Reset();
+	MarkComponentsRenderStateDirty();
+}
+
+void AVisualLoggerRenderingActor::OnFiltersChanged()
+{
+	const TMap<FName, FVisualLogExtensionInterface*>& AllExtensions = FVisualLogger::Get().GetAllExtensions();
+	for (auto& Extension : AllExtensions)
+	{
+		Extension.Value->DrawData(FVisualLoggerEditorInterface::Get(), NULL);
+	}
+
+	DebugShapesPerRow.Reset();
+	const TArray<FName>& RowNames = FVisualLoggerDatabase::Get().GetSelectedRows();
+	for (FName CurrentName : RowNames)
+	{
+		FVisualLoggerDBRow& DBRow = FVisualLoggerDatabase::Get().GetRowByName(CurrentName);
+		FTimelineDebugShapes& ShapesCache = DebugShapesPerRow.FindOrAdd(CurrentName);
+		ShapesCache.Reset();
+		if (DBRow.GetCurrentItemIndex() != INDEX_NONE)
+		{
+			GetDebugShapes(DBRow.GetCurrentItem(), ShapesCache);
+		}
+	}
+	MarkComponentsRenderStateDirty();
 }
 
 void AVisualLoggerRenderingActor::AddDebugRendering()
@@ -184,25 +295,6 @@ void AVisualLoggerRenderingActor::AddDebugRendering()
 	}
 }
 
-void AVisualLoggerRenderingActor::ObjectSelectionChanged(TArray<TSharedPtr<class STimeline> >& TimeLines)
-{
-	PrimaryDebugShapes.Reset();
-
-	if (TimeLines.Num() > 0)
-	{
-		const TArray<FVisualLogDevice::FVisualLogEntryItem>& Entries = TimeLines[TimeLines.Num() - 1]->GetEntries();
-		for (const auto &CurrentEntry : Entries)
-		{
-			if (CurrentEntry.Entry.Location != FVector::ZeroVector)
-			{
-				PrimaryDebugShapes.LogEntriesPath.Add(CurrentEntry.Entry.Location);
-			}
-		}
-	}
-	SelectedTimeLines = TimeLines;
-	MarkComponentsRenderStateDirty();
-}
-
 namespace
 {
 	static void GetPolygonMesh(const FVisualLogShapeElement* ElementToDraw, FDebugRenderSceneProxy::FMesh& TestMesh, const FVector& VertexOffset = FVector::ZeroVector)
@@ -265,35 +357,9 @@ void AVisualLoggerRenderingActor::GetDebugShapes(const FVisualLogDevice::FVisual
 		}
 	}
 
-	const TMap<FName, FVisualLogExtensionInterface*>& AllExtensions = FVisualLogger::Get().GetAllExtensions();
-	for (auto& Extension : AllExtensions)
-	{
-		Extension.Value->OnTimestampChange(Entry->TimeStamp, GetWorld(), this);
-	}
-
-	for (const auto CurrentData : Entry->DataBlocks)
-	{
-		const FName TagName = CurrentData.TagName;
-		const bool bIsValidByFilter = FCategoryFiltersManager::Get().MatchCategoryFilters(CurrentData.Category.ToString(), ELogVerbosity::All) && FCategoryFiltersManager::Get().MatchCategoryFilters(CurrentData.TagName.ToString(), ELogVerbosity::All);
-		FVisualLogExtensionInterface* Extension = FVisualLogger::Get().GetExtensionForTag(TagName);
-		if (!Extension)
-		{
-			continue;
-		}
-
-		if (!bIsValidByFilter)
-		{
-			Extension->DisableDrawingForData(GetWorld(), NULL, this, TagName, CurrentData, Entry->TimeStamp);
-		}
-		else
-		{
-			Extension->DrawData(GetWorld(), NULL, this, TagName, CurrentData, Entry->TimeStamp);
-		}
-	}
-
 	for (int32 ElementIndex = 0; ElementIndex < ElementsCount; ++ElementIndex, ++ElementToDraw)
 	{
-		if (!FCategoryFiltersManager::Get().MatchCategoryFilters(ElementToDraw->Category.ToString(), ElementToDraw->Verbosity))
+		if (!FVisualLoggerFilters::Get().MatchCategoryFilters(ElementToDraw->Category.ToString(), ElementToDraw->Verbosity))
 		{
 			continue;
 		}
@@ -557,55 +623,3 @@ void AVisualLoggerRenderingActor::GetDebugShapes(const FVisualLogDevice::FVisual
 		}
 	}
 }
-
-void AVisualLoggerRenderingActor::OnItemSelectionChanged(const FVisualLogDevice::FVisualLogEntryItem& EntryItem)
-{
-	PrimaryDebugShapes.Reset();
-	GetDebugShapes(EntryItem, PrimaryDebugShapes);
-
-	SecondaryDebugShapes.Reset();
-	if (SelectedTimeLines.Num() > 1)
-	{
-		for (TSharedPtr<class STimeline> CurrentTimeline : SelectedTimeLines)
-		{
-			if (CurrentTimeline.IsValid() == false || CurrentTimeline->GetVisibility().IsVisible() == false)
-			{
-				continue;
-			}
-
-			if (EntryItem.OwnerName == CurrentTimeline->GetName())
-			{
-				continue;
-			}
-
-			const TArray<FVisualLogDevice::FVisualLogEntryItem>& Entries = CurrentTimeline->GetEntries();
-
-			int32 BestItemIndex = INDEX_NONE;
-			float BestDistance = MAX_FLT;
-			for (int32 Index = 0; Index < Entries.Num(); Index++)
-			{
-				auto& CurrentEntryItem = Entries[Index];
-
-				if (CurrentTimeline->IsEntryHidden(CurrentEntryItem))
-				{
-					continue;
-				}
-				TArray<FVisualLoggerCategoryVerbosityPair> OutCategories;
-				const float CurrentDist = FMath::Abs(CurrentEntryItem.Entry.TimeStamp - EntryItem.Entry.TimeStamp);
-				if (CurrentDist < BestDistance)
-				{
-					BestDistance = CurrentDist;
-					BestItemIndex = Index;
-				}
-			}
-
-			if (Entries.IsValidIndex(BestItemIndex))
-			{
-				GetDebugShapes(Entries[BestItemIndex], SecondaryDebugShapes);
-			}
-		}
-	}
-
-	MarkComponentsRenderStateDirty();
-}
-
diff --git a/Engine/Source/Developer/LogVisualizer/Private/VisualLoggerRenderingActor.h b/Engine/Source/Developer/LogVisualizer/Private/VisualLoggerRenderingActor.h
--- a/Engine/Source/Developer/LogVisualizer/Private/VisualLoggerRenderingActor.h
+++ b/Engine/Source/Developer/LogVisualizer/Private/VisualLoggerRenderingActor.h
@@ -41,21 +41,22 @@ class LOGVISUALIZER_API AVisualLoggerRenderingActor : public AActor
 
 public:
 	GENERATED_UCLASS_BODY()
-
-	void OnItemSelectionChanged(const FVisualLogDevice::FVisualLogEntryItem& EntryItem);
-	void ObjectSelectionChanged(TArray<TSharedPtr<class STimeline> >& TimeLines);
-
-	void GetDebugShapes(const FVisualLogDevice::FVisualLogEntryItem& EntryItem, FTimelineDebugShapes& DebugShapes);
+	virtual ~AVisualLoggerRenderingActor();
+	void ResetRendering();
+	void ObjectSelectionChanged(const TArray<FName>& Selection);
+	void OnItemSelectionChanged(const FVisualLoggerDBRow& BDRow, int32 ItemIndex);
 
 private:
 	void AddDebugRendering();
+	void GetDebugShapes(const FVisualLogDevice::FVisualLogEntryItem& EntryItem, FTimelineDebugShapes& DebugShapes);
+	void OnFiltersChanged();
 
 public:
 	UPrimitiveComponent* RenderingComponent;
 
-	FTimelineDebugShapes PrimaryDebugShapes;
-	FTimelineDebugShapes SecondaryDebugShapes;
+	//FTimelineDebugShapes PrimaryDebugShapes;
 	FTimelineDebugShapes TestDebugShapes;
 
-	TArray<TSharedPtr<class STimeline> > SelectedTimeLines;
+	TArray<FName> CachedRowSelection;
+	TMap<FName, FTimelineDebugShapes> DebugShapesPerRow;
 };
diff --git a/Engine/Source/Developer/LogVisualizer/Public/ILogVisualizer.h b/Engine/Source/Developer/LogVisualizer/Public/ILogVisualizer.h
--- a/Engine/Source/Developer/LogVisualizer/Public/ILogVisualizer.h
+++ b/Engine/Source/Developer/LogVisualizer/Public/ILogVisualizer.h
@@ -5,7 +5,4 @@
 class ILogVisualizer : public IModuleInterface, public IModularFeature
 {
 public:
-	virtual void Goto(float Timestamp, FName LogOwner = NAME_None) = 0;
-	virtual void GotoNextItem() = 0;
-	virtual void GotoPreviousItem() = 0;
 };
\ No newline at end of file
diff --git a/Engine/Source/Developer/LogVisualizer/Public/LogVisualizerSettings.h b/Engine/Source/Developer/LogVisualizer/Public/LogVisualizerSettings.h
--- a/Engine/Source/Developer/LogVisualizer/Public/LogVisualizerSettings.h
+++ b/Engine/Source/Developer/LogVisualizer/Public/LogVisualizerSettings.h
@@ -5,9 +5,12 @@ GameplayDebuggerSettings.h: Declares the UGameplayDebuggerSettings class.
 =============================================================================*/
 #pragma once
 
-
+#include "VisualLoggerDatabase.h"
 #include "LogVisualizerSettings.generated.h"
 
+DECLARE_MULTICAST_DELEGATE_TwoParams(FOnFilterCategoryAdded, FString, ELogVerbosity::Type);
+DECLARE_MULTICAST_DELEGATE_OneParam(FOnFilterCategoryRemoved, FString);
+
 USTRUCT()
 struct FCategoryFilter
 {
@@ -24,7 +27,7 @@ struct FCategoryFilter
 };
 
 USTRUCT()
-struct FVisualLoggerFilters
+struct FVisualLoggerFiltersData
 {
 	GENERATED_USTRUCT_BODY()
 
@@ -41,6 +44,51 @@ struct FVisualLoggerFilters
 	TArray<FString> SelectedClasses;
 };
 
+USTRUCT()
+struct FVisualLoggerFilters : public FVisualLoggerFiltersData
+{
+	GENERATED_USTRUCT_BODY()
+
+	FOnFilterCategoryAdded OnFilterCategoryAdded;
+	FOnFilterCategoryRemoved OnFilterCategoryRemoved;
+
+	static FVisualLoggerFilters& Get();
+	static void Initialize();
+	static void Shutdown();
+
+	void Reset();
+
+	bool MatchCategoryFilters(FString String, ELogVerbosity::Type Verbosity = ELogVerbosity::All);
+
+	bool MatchSearchString(FString String) { return SearchBoxFilter == String; }
+	void SetSearchString(FString InString) { SearchBoxFilter = InString; }
+	FString GetSearchString() { return SearchBoxFilter; }
+
+	void AddCategory(FString InName, ELogVerbosity::Type InVerbosity);
+	void RemoveCategory(FString InName);
+	FCategoryFilter& GetCategoryByName(const FString& InName);
+	FCategoryFilter& GetCategoryByName(const FName& InName);
+
+	void DeactivateAllButThis(const FString& InName);
+	void EnableAllCategories();
+
+	bool MatchObjectName(FString String);
+	void SelectObject(FString ObjectName);
+	void RemoveObjectFromSelection(FString ObjectName);
+	const TArray<FString>& GetSelectedObjects() const;
+
+	void DisableGraphData(FName GraphName, FName DataName, bool SetAsDisabled);
+	bool IsGraphDataDisabled(FName GraphName, FName DataName);
+
+protected:
+	void OnNewItemHandler(const FVisualLoggerDBRow& BDRow, int32 ItemIndex);
+
+private:
+	static TSharedPtr< struct FVisualLoggerFilters > StaticInstance;
+	TMap<FName, FCategoryFilter*>	FastCategoryFilterMap;
+	TArray<FName> DisabledGraphDatas;
+};
+
 struct FCategoryFiltersManager;
 
 UCLASS(config = EditorPerProjectUserSettings)
@@ -108,11 +156,15 @@ public:
 
 	class UMaterial* GetDebugMeshMaterial();
 
+	void SavePresistentData();
+
+	void ClearPresistentData();
+
+	void LoadPresistentData();
+
 protected:
-	FVisualLoggerFilters CurrentFilters;
-	
 	UPROPERTY(config)
-	FVisualLoggerFilters PresistentFilters;
+	FVisualLoggerFiltersData PresistentFilters;
 
 	/** A material used to render debug meshes with kind of flat shading, mostly used by Visual Logger tool. */
 	UPROPERTY()
@@ -127,34 +179,3 @@ private:
 	FSettingChangedEvent SettingChangedEvent;
 
 };
-
-struct FCategoryFiltersManager
-{
-	static FCategoryFiltersManager& Get() { return StaticManager; }
-
-	bool MatchCategoryFilters(FString String, ELogVerbosity::Type Verbosity = ELogVerbosity::All);
-	bool MatchObjectName(FString String);
-	bool MatchSearchString(FString String);
-
-	void SetSearchString(FString InString);
-	FString GetSearchString();
-
-	void SetObjectFilterString(FString InFilterString);
-	FString GetObjectFilterString();
-
-	void AddCategory(FString InName, ELogVerbosity::Type InVerbosity);
-	void RemoveCategory(FString InName);
-	bool IsValidCategory(FString InName);
-	FCategoryFilter& GetCategory(FString InName);
-
-	void SelectObject(FString ObjectName);
-	void RemoveObjectFromSelection(FString ObjectName);
-	const TArray<FString>& GetSelectedObjects();
-
-	void SavePresistentData();
-	void ClearPresistentData();
-	void LoadPresistentData();
-
-protected:
-	static FCategoryFiltersManager StaticManager;
-};
diff --git a/Engine/Source/Runtime/AIModule/Classes/AIController.h b/Engine/Source/Runtime/AIModule/Classes/AIController.h
--- a/Engine/Source/Runtime/AIModule/Classes/AIController.h
+++ b/Engine/Source/Runtime/AIModule/Classes/AIController.h
@@ -14,6 +14,7 @@
 #include "BehaviorTree/BehaviorTreeTypes.h"
 #include "GameplayTaskOwnerInterface.h"
 #include "GenericTeamAgentInterface.h"
+#include "VisualLogger/VisualLoggerDebugSnapshotInterface.h"
 #include "AIController.generated.h"
 
 class APawn;
@@ -77,7 +78,7 @@ struct FFocusKnowledge
  */
 
 UCLASS(ClassGroup = AI, BlueprintType, Blueprintable)
-class AIMODULE_API AAIController : public AController, public IAIPerceptionListenerInterface, public IGameplayTaskOwnerInterface, public IGenericTeamAgentInterface
+class AIMODULE_API AAIController : public AController, public IAIPerceptionListenerInterface, public IGameplayTaskOwnerInterface, public IGenericTeamAgentInterface, public IVisualLoggerDebugSnapshotInterface
 {
 	GENERATED_BODY()
 
diff --git a/Engine/Source/Runtime/AIModule/Classes/VisualLoggerExtension.h b/Engine/Source/Runtime/AIModule/Classes/VisualLoggerExtension.h
--- a/Engine/Source/Runtime/AIModule/Classes/VisualLoggerExtension.h
+++ b/Engine/Source/Runtime/AIModule/Classes/VisualLoggerExtension.h
@@ -22,18 +22,20 @@ class FVisualLoggerExtension : public FVisualLogExtensionInterface
 {
 public:
 	FVisualLoggerExtension();
-	virtual void OnTimestampChange(float Timestamp, UWorld* InWorld, AActor* HelperActor) override;
-	virtual void DrawData(UWorld* InWorld, UCanvas* Canvas, AActor* HelperActor, const FName& TagName, const FVisualLogDataBlock& DataBlock, float Timestamp) override;
-	virtual void DisableDrawingForData(UWorld* InWorld, UCanvas* Canvas, AActor* HelperActor, const FName& TagName, const FVisualLogDataBlock& DataBlock, float Timestamp) override;
-	virtual void LogEntryLineSelectionChanged(TSharedPtr<FLogEntryItem> SelectedItem, int64 UserData, FName TagName) override;
+
+	virtual void ResetData(IVisualLoggerEditorInterface* EdInterface) override;
+	virtual void DrawData(IVisualLoggerEditorInterface* EdInterface, UCanvas* Canvas) override;
+	virtual void OnItemsSelectionChanged(IVisualLoggerEditorInterface* EdInterface) override;
+	virtual void OnLogLineSelectionChanged(IVisualLoggerEditorInterface* EdInterface, TSharedPtr<struct FLogEntryItem> SelectedItem, int64 UserData) override;
 
 private:
+	void DrawData(UWorld* InWorld, class UEQSRenderingComponent* EQSRenderingComponent, UCanvas* Canvas, AActor* HelperActor, const FName& TagName, const FVisualLogDataBlock& DataBlock, float Timestamp);
 	void DisableEQSRendering(AActor* HelperActor);
 
 protected:
-	uint32 CachedEQSId;
 	uint32 SelectedEQSId;
 	float CurrentTimestamp;
+	TArray<TWeakObjectPtr<class UEQSRenderingComponent> >	EQSRenderingComponents;
 };
 #endif //ENABLE_VISUAL_LOG
 
diff --git a/Engine/Source/Runtime/AIModule/Private/VisualLoggerExtension.cpp b/Engine/Source/Runtime/AIModule/Private/VisualLoggerExtension.cpp
--- a/Engine/Source/Runtime/AIModule/Private/VisualLoggerExtension.cpp
+++ b/Engine/Source/Runtime/AIModule/Private/VisualLoggerExtension.cpp
@@ -10,35 +10,12 @@
 
 #if ENABLE_VISUAL_LOG
 FVisualLoggerExtension::FVisualLoggerExtension()
-	: CachedEQSId(INDEX_NONE)
-	, SelectedEQSId(INDEX_NONE)
+	: SelectedEQSId(INDEX_NONE)
 	, CurrentTimestamp(FLT_MIN)
 {
 
 }
 
-void FVisualLoggerExtension::OnTimestampChange(float Timestamp, UWorld* InWorld, AActor* HelperActor )
-{
-#if USE_EQS_DEBUGGER
-	if (CurrentTimestamp != Timestamp)
-	{
-		CurrentTimestamp = Timestamp;
-		CachedEQSId = INDEX_NONE;
-		DisableEQSRendering(HelperActor);
-	}
-#endif
-}
-
-void FVisualLoggerExtension::DisableDrawingForData(UWorld* InWorld, UCanvas* Canvas, AActor* HelperActor, const FName& TagName, const FVisualLogDataBlock& DataBlock, float Timestamp)
-{
-#if USE_EQS_DEBUGGER
-	if (TagName == *EVisLogTags::TAG_EQS && CurrentTimestamp == Timestamp)
-	{
-		DisableEQSRendering(HelperActor);
-	}
-#endif
-}
-
 void FVisualLoggerExtension::DisableEQSRendering(AActor* HelperActor)
 {
 #if USE_EQS_DEBUGGER
@@ -57,47 +34,95 @@ void FVisualLoggerExtension::DisableEQSRendering(AActor* HelperActor)
 #endif
 }
 
-void FVisualLoggerExtension::LogEntryLineSelectionChanged(TSharedPtr<FLogEntryItem> SelectedItem, int64 UserData, FName TagName)
+extern RENDERCORE_API FTexture* GWhiteTexture;
+
+void FVisualLoggerExtension::ResetData(IVisualLoggerEditorInterface* EdInterface)
 {
-	if (TagName == *EVisLogTags::TAG_EQS)
-	{
-		SelectedEQSId = (int32)UserData;
-	}
+	DisableEQSRendering(EdInterface->GetHelperActor());
 }
 
-extern RENDERCORE_API FTexture* GWhiteTexture;
+void FVisualLoggerExtension::OnItemsSelectionChanged(IVisualLoggerEditorInterface* EdInterface)
+{
+	DrawData(EdInterface, nullptr); //it'll reset current data first
+}
 
-void FVisualLoggerExtension::DrawData(UWorld* InWorld, UCanvas* Canvas, AActor* HelperActor, const FName& TagName, const FVisualLogDataBlock& DataBlock, float Timestamp)
+void FVisualLoggerExtension::OnLogLineSelectionChanged(IVisualLoggerEditorInterface* EdInterface, TSharedPtr<struct FLogEntryItem> SelectedItem, int64 UserData)
 {
-#if USE_EQS_DEBUGGER
-	if (TagName == *EVisLogTags::TAG_EQS)
+	SelectedEQSId = SelectedItem.IsValid() ? UserData : INDEX_NONE;
+	EdInterface->GetHelperActor()->MarkComponentsRenderStateDirty();
+	DrawData(EdInterface, NULL); //we have to refresh rendering components
+}
+
+void FVisualLoggerExtension::DrawData(IVisualLoggerEditorInterface* EdInterface, UCanvas* Canvas)
+{
+	if (Canvas == nullptr)
 	{
-		EQSDebug::FQueryData DebugData;
-		UEnvQueryDebugHelpers::BlobArrayToDebugData(DataBlock.Data, DebugData, false);
-		if (HelperActor)
+		// disable and refresh EQS rendering
+		DisableEQSRendering(EdInterface->GetHelperActor());
+	}
+	int32 EQSRenderingComponentIndex = 0;
+	TArray<FName> RowNames = EdInterface->GetSelectedRows();
+	for (const FName& RowName : RowNames)
+	{
+		if (EdInterface->GetSelectedItemIndex(RowName) != INDEX_NONE )
 		{
-			UEQSRenderingComponent* EQSRenderComp = HelperActor->FindComponentByClass<UEQSRenderingComponent>();
-			if (!EQSRenderComp)
+			if (EdInterface->IsItemVisible(RowName, EdInterface->GetSelectedItemIndex(RowName)) == false)
 			{
-				EQSRenderComp = NewObject<UEQSRenderingComponent>(HelperActor);
-				EQSRenderComp->bDrawOnlyWhenSelected = false;
-				EQSRenderComp->RegisterComponent();
-				EQSRenderComp->SetHiddenInGame(true);
+				continue;
 			}
 
-			if (SelectedEQSId == INDEX_NONE || SelectedEQSId != CachedEQSId /*|| (EQSRenderComp && EQSRenderComp->bHiddenInGame)*/)
+			const FVisualLogDevice::FVisualLogEntryItem& EntryItem = EdInterface->GetSelectedItem(RowName);
+			if (Canvas == nullptr)
+			{
+				for (auto& Component : EQSRenderingComponents)
+				{
+					if (Component.IsValid())
+					{
+						Component->SetHiddenInGame(true);
+						Component->Deactivate();
+						Component->DebugData.Reset();
+						Component->MarkRenderStateDirty();
+					}
+				}
+			}
+			for (const auto& CurrentData : EntryItem.Entry.DataBlocks)
 			{
-				if (SelectedEQSId == INDEX_NONE)
+				const FName TagName = CurrentData.TagName;
+				if (EdInterface->MatchCategoryFilters(TagName.ToString()))
 				{
-					SelectedEQSId = DebugData.Id;
+					UEQSRenderingComponent* EQSRenderComp = (Canvas == nullptr && EQSRenderingComponents.IsValidIndex(EQSRenderingComponentIndex) ? EQSRenderingComponents[EQSRenderingComponentIndex].Get() : nullptr);
+					if (EQSRenderComp == nullptr && Canvas == nullptr)
+					{
+						EQSRenderComp = NewObject<UEQSRenderingComponent>(EdInterface->GetHelperActor());
+						EQSRenderComp->bDrawOnlyWhenSelected = false;
+						EQSRenderComp->RegisterComponent();
+						EQSRenderComp->SetHiddenInGame(false);
+						EQSRenderComp->Activate();
+						EQSRenderComp->MarkRenderStateDirty();
+						EQSRenderingComponents.Add(EQSRenderComp);
+					}
+					DrawData(EdInterface->GetWorld(), EQSRenderComp, Canvas, EdInterface->GetHelperActor(), TagName, CurrentData, EntryItem.Entry.TimeStamp);
 				}
-				CachedEQSId = DebugData.Id;
-				EQSRenderComp->DebugData = DebugData;
-				EQSRenderComp->Activate();
-				EQSRenderComp->SetHiddenInGame(false);
-				EQSRenderComp->MarkRenderStateDirty();
+				EQSRenderingComponentIndex++;
 			}
 		}
+	}
+}
+
+void FVisualLoggerExtension::DrawData(UWorld* InWorld, class UEQSRenderingComponent* EQSRenderComp, UCanvas* Canvas, AActor* HelperActor, const FName& TagName, const FVisualLogDataBlock& DataBlock, float Timestamp)
+{
+#if USE_EQS_DEBUGGER
+	if (TagName == *EVisLogTags::TAG_EQS)
+	{
+		EQSDebug::FQueryData DebugData;
+		UEnvQueryDebugHelpers::BlobArrayToDebugData(DataBlock.Data, DebugData, false);
+		if (EQSRenderComp && !Canvas && (SelectedEQSId == DebugData.Id || SelectedEQSId == INDEX_NONE))
+		{
+			EQSRenderComp->DebugData = DebugData;
+			EQSRenderComp->SetHiddenInGame(false);
+			EQSRenderComp->Activate();
+			EQSRenderComp->MarkRenderStateDirty();
+		}
 
 		/** find and draw item selection */
 		int32 BestItemIndex = INDEX_NONE;
diff --git a/Engine/Source/Runtime/Engine/Classes/GameFramework/Actor.h b/Engine/Source/Runtime/Engine/Classes/GameFramework/Actor.h
--- a/Engine/Source/Runtime/Engine/Classes/GameFramework/Actor.h
+++ b/Engine/Source/Runtime/Engine/Classes/GameFramework/Actor.h
@@ -2556,17 +2556,6 @@ public:
 	/** Retrieves actor's name used for logging, or string "NULL" if Actor == NULL */
 	static FString GetDebugName(const AActor* Actor) { return Actor ? Actor->GetName() : TEXT("NULL"); }
 
-#if ENABLE_VISUAL_LOG
-	/** 
-	 *	Hook for Actors to supply visual logger with additional data.
-	 *	It's guaranteed that Snapshot != NULL
-	 */
-	virtual void GrabDebugSnapshot(struct FVisualLogEntry* Snapshot) const {}
-
-private:
-	friend class FVisualLog;
-#endif // ENABLE_VISUAL_LOG
-
 	//* Sets the friendly actor label and name */
 	void SetActorLabelInternal( const FString& NewActorLabelDirty, bool bMakeGloballyUniqueFName, bool bMarkDirty );
 
diff --git a/Engine/Source/Runtime/Engine/Private/VisualLogger/VisualLogger.cpp b/Engine/Source/Runtime/Engine/Private/VisualLogger/VisualLogger.cpp
--- a/Engine/Source/Runtime/Engine/Private/VisualLogger/VisualLogger.cpp
+++ b/Engine/Source/Runtime/Engine/Private/VisualLogger/VisualLogger.cpp
@@ -3,6 +3,7 @@
 #include "EnginePrivate.h"
 #include "VisualLogger/VisualLogger.h"
 #include "VisualLogger/VisualLoggerBinaryFileDevice.h"
+#include "VisualLogger/VisualLoggerDebugSnapshotInterface.h"
 #if WITH_EDITOR
 #	include "Editor/UnrealEd/Public/EditorComponents.h"
 #	include "Editor/UnrealEd/Public/EditorReimportHandler.h"
@@ -18,6 +19,29 @@ DEFINE_LOG_CATEGORY(LogVisual);
 #if ENABLE_VISUAL_LOG 
 DEFINE_STAT(STAT_VisualLog);
 
+namespace
+{
+	static UWorld* GetWorldForVisualLogger(const class UObject* Object)
+	{
+		UWorld* World = Object ? GEngine->GetWorldFromContextObject(Object, false) : nullptr;
+#if WITH_EDITOR
+		UEditorEngine *EEngine = Cast<UEditorEngine>(GEngine);
+		if (GIsEditor && EEngine != nullptr && World == nullptr)
+		{
+			// lets use PlayWorld during PIE/Simulate and regular world from editor otherwise, to draw debug information
+			World = EEngine->PlayWorld != nullptr ? EEngine->PlayWorld : EEngine->GetEditorWorldContext().World();
+		}
+
+#endif
+		if (!GIsEditor && World == nullptr)
+		{
+			World = GEngine->GetWorld();
+		}
+
+		return World;
+	}
+}
+
 TMap<UObject*, TArray<TWeakObjectPtr<const UObject> > > FVisualLogger::RedirectionMap;
 int32 FVisualLogger::bIsRecording = false;
 
@@ -67,7 +91,7 @@ FVisualLogEntry* FVisualLogger::GetEntryToWrite(const class UObject* Object, flo
 
 	bool InitializeNewEntry = false;
 
-	TWeakObjectPtr<UWorld> World = GetWorld(Object);
+	TWeakObjectPtr<UWorld> World = GetWorldForVisualLogger(Object);
 
 	if (CurrentEntryPerObject.Contains(LogOwner))
 	{
@@ -110,20 +134,20 @@ FVisualLogEntry* FVisualLogger::GetEntryToWrite(const class UObject* Object, flo
 		{
 			if (ObjectToPointerMap.Contains(LogOwner) && ObjectToPointerMap[LogOwner].IsValid())
 			{
-				const class AActor* LogOwnerAsActor = Cast<class AActor>(LogOwner);
-				if (LogOwnerAsActor)
+				const IVisualLoggerDebugSnapshotInterface* DebugSnapshotInterface = Cast<const IVisualLoggerDebugSnapshotInterface>(LogOwner);
+				if (DebugSnapshotInterface)
 				{
-					LogOwnerAsActor->GrabDebugSnapshot(CurrentEntry);
+					DebugSnapshotInterface->GrabDebugSnapshot(CurrentEntry);
 				}
 			}
 			for (auto Child : RedirectionMap[LogOwner])
 			{
 				if (Child.IsValid())
 				{
-					const class AActor* ChildAsActor = Cast<class AActor>(Child.Get());
-					if (ChildAsActor)
+					const IVisualLoggerDebugSnapshotInterface* DebugSnapshotInterface = Cast<const IVisualLoggerDebugSnapshotInterface>(Child.Get());
+					if (DebugSnapshotInterface)
 					{
-						ChildAsActor->GrabDebugSnapshot(CurrentEntry);
+						DebugSnapshotInterface->GrabDebugSnapshot(CurrentEntry);
 					}
 				}
 			}
@@ -134,7 +158,11 @@ FVisualLogEntry* FVisualLogger::GetEntryToWrite(const class UObject* Object, flo
 			if (ObjectAsActor)
 			{
 				CurrentEntry->Location = ObjectAsActor->GetActorLocation();
-				ObjectAsActor->GrabDebugSnapshot(CurrentEntry);
+				const IVisualLoggerDebugSnapshotInterface* DebugSnapshotInterface = Cast<const IVisualLoggerDebugSnapshotInterface>(Object);
+				if (DebugSnapshotInterface)
+				{
+					DebugSnapshotInterface->GrabDebugSnapshot(CurrentEntry);
+				}
 			}
 		}
 	}
@@ -265,7 +293,7 @@ void FVisualLogger::NavigationDataDump(const class UObject* Object, const struct
 	const FNavDataGenerator* Generator = MainNavData ? MainNavData->GetGenerator() : nullptr;
 	if (Generator)
 	{
-		Generator->GrabDebugSnapshot(CurrentEntry, FMath::IsNearlyZero(Box.GetVolume()) ? MainNavData->GetBounds() : Box, Category, Verbosity);
+		Generator->GrabDebugSnapshot(CurrentEntry, FMath::IsNearlyZero(Box.GetVolume()) ? MainNavData->GetBounds().ExpandBy(FVector(20,20,20)) : Box, Category, Verbosity);
 	}
 }
 
@@ -284,34 +312,6 @@ FVisualLogger::FVisualLogger()
 	}
 }
 
-namespace
-{
-	static UWorld* GetWorldForVisualLogger(const class UObject* Object)
-	{
-		UWorld* World = Object ? GEngine->GetWorldFromContextObject(Object, false) : nullptr;
-#if WITH_EDITOR
-		UEditorEngine *EEngine = Cast<UEditorEngine>(GEngine);
-		if (GIsEditor && EEngine != nullptr && World == nullptr)
-		{
-			// lets use PlayWorld during PIE/Simulate and regular world from editor otherwise, to draw debug information
-			World = EEngine->PlayWorld != nullptr ? EEngine->PlayWorld : EEngine->GetEditorWorldContext().World();
-		}
-
-#endif
-		if (!GIsEditor && World == nullptr)
-		{
-			World = GEngine->GetWorld();
-		}
-
-		return World;
-	}
-}
-
-UWorld* FVisualLogger::GetWorld(const class UObject* Object)
-{
-	return GetWorldForVisualLogger(Object);
-}
-
 void FVisualLogger::Shutdown()
 {
 	SetIsRecording(false);
@@ -395,7 +395,7 @@ void FVisualLogger::SetIsRecording(bool InIsRecording)
 { 
 	if (InIsRecording == false && InIsRecording != !!bIsRecording && FParse::Param(FCommandLine::Get(), TEXT("LogNavOctree")))
 	{
-		FVisualLogger::NavigationDataDump(GetWorld(nullptr), LogNavigation, ELogVerbosity::Log, INDEX_NONE, FBox());
+		FVisualLogger::NavigationDataDump(GetWorldForVisualLogger(nullptr), LogNavigation, ELogVerbosity::Log, INDEX_NONE, FBox());
 	}
 	if (IsRecordingToFile())
 	{
diff --git a/Engine/Source/Runtime/Engine/Private/VisualLogger/VisualLoggerTypes.cpp b/Engine/Source/Runtime/Engine/Private/VisualLogger/VisualLoggerTypes.cpp
--- a/Engine/Source/Runtime/Engine/Private/VisualLogger/VisualLoggerTypes.cpp
+++ b/Engine/Source/Runtime/Engine/Private/VisualLogger/VisualLoggerTypes.cpp
@@ -3,6 +3,12 @@
 #include "EnginePrivate.h"
 #include "ObjectBase.h"
 #include "VisualLogger/VisualLogger.h"
+#include "VisualLogger/VisualLoggerDebugSnapshotInterface.h"
+
+UVisualLoggerDebugSnapshotInterface::UVisualLoggerDebugSnapshotInterface(const FObjectInitializer& ObjectInitializer)
+	: Super(ObjectInitializer)
+{
+}
 
 #if ENABLE_VISUAL_LOG
 
@@ -44,7 +50,11 @@ FVisualLogEntry::FVisualLogEntry(const class AActor* InActor, TArray<TWeakObject
 	{
 		TimeStamp = InActor->GetWorld()->TimeSeconds;
 		Location = InActor->GetActorLocation();
-		InActor->GrabDebugSnapshot(this);
+		const IVisualLoggerDebugSnapshotInterface* DebugSnapshotInterface = Cast<const IVisualLoggerDebugSnapshotInterface>(InActor);
+		if (DebugSnapshotInterface)
+		{
+			DebugSnapshotInterface->GrabDebugSnapshot(this);
+		}
 		if (Children != nullptr)
 		{
 			TWeakObjectPtr<UObject>* WeakActorPtr = Children->GetData();
@@ -52,7 +62,7 @@ FVisualLogEntry::FVisualLogEntry(const class AActor* InActor, TArray<TWeakObject
 			{
 				if (WeakActorPtr->IsValid())
 				{
-					const AActor* ChildActor = Cast<AActor>(WeakActorPtr->Get());
+					const IVisualLoggerDebugSnapshotInterface* ChildActor = Cast<const IVisualLoggerDebugSnapshotInterface>(WeakActorPtr->Get());
 					if (ChildActor)
 					{
 						ChildActor->GrabDebugSnapshot(this);
@@ -67,10 +77,10 @@ FVisualLogEntry::FVisualLogEntry(float InTimeStamp, FVector InLocation, const UO
 {
 	TimeStamp = InTimeStamp;
 	Location = InLocation;
-	const AActor* AsActor = Cast<AActor>(Object);
-	if (AsActor)
+	const IVisualLoggerDebugSnapshotInterface* DebugSnapshotInterface = Cast<const IVisualLoggerDebugSnapshotInterface>(Object);
+	if (DebugSnapshotInterface)
 	{
-		AsActor->GrabDebugSnapshot(this);
+		DebugSnapshotInterface->GrabDebugSnapshot(this);
 	}
 	if (Children != nullptr)
 	{
@@ -79,7 +89,7 @@ FVisualLogEntry::FVisualLogEntry(float InTimeStamp, FVector InLocation, const UO
 		{
 			if (WeakActorPtr->IsValid())
 			{
-				const AActor* ChildActor = Cast<AActor>(WeakActorPtr->Get());
+				const IVisualLoggerDebugSnapshotInterface* ChildActor = Cast<const IVisualLoggerDebugSnapshotInterface>(WeakActorPtr->Get());
 				if (ChildActor)
 				{
 					ChildActor->GrabDebugSnapshot(this);
@@ -505,27 +515,27 @@ void FVisualLoggerHelpers::GetCategories(const FVisualLogEntry& EntryItem, TArra
 {
 	for (const auto& CurrentEvent : EntryItem.Events)
 	{
-		OutCategories.Add(FVisualLoggerCategoryVerbosityPair(*CurrentEvent.Name, CurrentEvent.Verbosity));
+		OutCategories.AddUnique(FVisualLoggerCategoryVerbosityPair(*CurrentEvent.Name, ELogVerbosity::All));
 	}
 
 	for (const auto& CurrentLine : EntryItem.LogLines)
 	{
-		OutCategories.Add(FVisualLoggerCategoryVerbosityPair(CurrentLine.Category, CurrentLine.Verbosity));
+		OutCategories.AddUnique(FVisualLoggerCategoryVerbosityPair(CurrentLine.Category, ELogVerbosity::All));
 	}
 
 	for (const auto& CurrentElement : EntryItem.ElementsToDraw)
 	{
-		OutCategories.Add(FVisualLoggerCategoryVerbosityPair(CurrentElement.Category, CurrentElement.Verbosity));
+		OutCategories.AddUnique(FVisualLoggerCategoryVerbosityPair(CurrentElement.Category, ELogVerbosity::All));
 	}
 
 	for (const auto& CurrentSample : EntryItem.HistogramSamples)
 	{
-		OutCategories.Add(FVisualLoggerCategoryVerbosityPair(CurrentSample.Category, CurrentSample.Verbosity));
+		OutCategories.AddUnique(FVisualLoggerCategoryVerbosityPair(CurrentSample.Category, ELogVerbosity::All));
 	}
 
 	for (const auto& CurrentBlock : EntryItem.DataBlocks)
 	{
-		OutCategories.Add(FVisualLoggerCategoryVerbosityPair(CurrentBlock.Category, CurrentBlock.Verbosity));
+		OutCategories.AddUnique(FVisualLoggerCategoryVerbosityPair(CurrentBlock.Category, ELogVerbosity::All));
 	}
 }
 
@@ -536,7 +546,7 @@ void FVisualLoggerHelpers::GetHistogramCategories(const FVisualLogEntry& EntryIt
 		auto& DataNames = OutCategories.FindOrAdd(CurrentSample.GraphName.ToString());
 		if (DataNames.Find(CurrentSample.DataName.ToString()) == INDEX_NONE)
 		{
-			DataNames.Add(CurrentSample.DataName.ToString());
+			DataNames.AddUnique(CurrentSample.DataName.ToString());
 		}
 	}
 }
diff --git a/Engine/Source/Runtime/Engine/Public/VisualLogger/VisualLogger.h b/Engine/Source/Runtime/Engine/Public/VisualLogger/VisualLogger.h
--- a/Engine/Source/Runtime/Engine/Public/VisualLogger/VisualLogger.h
+++ b/Engine/Source/Runtime/Engine/Public/VisualLogger/VisualLogger.h
@@ -64,7 +64,7 @@
 #define UE_VLOG_EVENT_WITH_DATA(LogOwner, LogEvent, ...) if(FVisualLogger::IsRecording()) FVisualLogger::EventLog(LogOwner, LogEvent, ##__VA_ARGS__)
 #define UE_CVLOG_EVENT_WITH_DATA(Condition, LogOwner, LogEvent, ...) if(FVisualLogger::IsRecording() && Condition) {UE_VLOG_EVENT_WITH_DATA(LogOwner, LogEvent, ##__VA_ARGS__);}
 
-#define UE_IFVLOG(__code_block__) __code_block__;
+#define UE_IFVLOG(__code_block__) if( FVisualLogger::IsRecording() ) { __code_block__; }
 
 #else
 #define REDIRECT_TO_VLOG(Dest)
@@ -252,7 +252,6 @@ public:
 private:
 	FVisualLogger();
 	virtual void Serialize(const TCHAR* V, ELogVerbosity::Type Verbosity, const class FName& Category) override { ensureMsgf(0, TEXT("Regular serialize is forbiden for visual logs")); }
-	class UWorld* GetWorld(const class UObject* Object);
 
 protected:
 	/** Array of output devices to redirect to */
diff --git a/Engine/Source/Runtime/Engine/Public/VisualLogger/VisualLoggerDebugSnapshotInterface.h b/Engine/Source/Runtime/Engine/Public/VisualLogger/VisualLoggerDebugSnapshotInterface.h
new file mode 100644
--- /dev/null
+++ b/Engine/Source/Runtime/Engine/Public/VisualLogger/VisualLoggerDebugSnapshotInterface.h
@@ -0,0 +1,16 @@
+#pragma once
+#include "VisualLoggerDebugSnapshotInterface.generated.h"
+
+UINTERFACE(MinimalAPI, meta = (CannotImplementInterfaceInBlueprint))
+class UVisualLoggerDebugSnapshotInterface : public UInterface
+{
+	GENERATED_UINTERFACE_BODY()
+};
+
+class ENGINE_API IVisualLoggerDebugSnapshotInterface
+{
+	GENERATED_IINTERFACE_BODY()
+public:
+	virtual void GrabDebugSnapshot(struct FVisualLogEntry* Snapshot) const {}
+};
+
diff --git a/Engine/Source/Runtime/Engine/Public/VisualLogger/VisualLoggerTypes.h b/Engine/Source/Runtime/Engine/Public/VisualLogger/VisualLoggerTypes.h
--- a/Engine/Source/Runtime/Engine/Public/VisualLogger/VisualLoggerTypes.h
+++ b/Engine/Source/Runtime/Engine/Public/VisualLogger/VisualLoggerTypes.h
@@ -219,15 +219,6 @@ struct ENGINE_API FVisualLogEntry
 
 #if  ENABLE_VISUAL_LOG
 
-class FVisualLogExtensionInterface
-{
-public:
-	virtual void OnTimestampChange(float Timestamp, class UWorld* InWorld, class AActor* HelperActor) = 0;
-	virtual void DrawData(class UWorld* InWorld, class UCanvas* Canvas, class AActor* HelperActor, const FName& TagName, const FVisualLogDataBlock& DataBlock, float Timestamp) = 0;
-	virtual void DisableDrawingForData(class UWorld* InWorld, class UCanvas* Canvas, class AActor* HelperActor, const FName& TagName, const FVisualLogDataBlock& DataBlock, float Timestamp) = 0;
-	virtual void LogEntryLineSelectionChanged(TSharedPtr<struct FLogEntryItem> SelectedItem, int64 UserData, FName TagName) = 0;
-};
-
 /**
  * Interface for Visual Logger Device
  */
@@ -262,6 +253,12 @@ struct ENGINE_API FVisualLoggerCategoryVerbosityPair
 	ELogVerbosity::Type Verbosity;
 };
 
+inline bool operator==(const FVisualLoggerCategoryVerbosityPair& A, const FVisualLoggerCategoryVerbosityPair& B)
+{
+	return A.CategoryName == B.CategoryName
+		&& A.Verbosity == B.Verbosity;
+}
+
 struct ENGINE_API FVisualLoggerHelpers
 {
 	static FString GenerateTemporaryFilename(const FString& FileExt);
@@ -272,6 +269,32 @@ struct ENGINE_API FVisualLoggerHelpers
 	static void GetHistogramCategories(const FVisualLogEntry& RecordedLogs, TMap<FString, TArray<FString> >& OutCategories);
 };
 
+struct IVisualLoggerEditorInterface
+{
+	virtual const FName& GetRowClassName(FName RowName) const = 0;
+	virtual int32 GetSelectedItemIndex(FName RowName) const = 0;
+	virtual const TArray<FVisualLogDevice::FVisualLogEntryItem>& GetRowItems(FName RowName) = 0;
+	virtual const FVisualLogDevice::FVisualLogEntryItem& GetSelectedItem(FName RowName) const = 0;
+
+	virtual const TArray<FName>& GetSelectedRows() const = 0;
+	virtual bool IsRowVisible(FName RowName) const = 0;
+	virtual bool IsItemVisible(FName RowName, int32 ItemIndex) const = 0;
+	virtual UWorld* GetWorld() const = 0;
+	virtual AActor* GetHelperActor() const = 0;
+
+	virtual bool MatchCategoryFilters(const FString& String, ELogVerbosity::Type Verbosity = ELogVerbosity::All) = 0;
+};
+
+class FVisualLogExtensionInterface
+{
+public:
+	virtual void ResetData(IVisualLoggerEditorInterface* EdInterface) = 0;
+	virtual void DrawData(IVisualLoggerEditorInterface* EdInterface, UCanvas* Canvas) = 0;
+
+	virtual void OnItemsSelectionChanged(IVisualLoggerEditorInterface* EdInterface) {};
+	virtual void OnLogLineSelectionChanged(IVisualLoggerEditorInterface* EdInterface, TSharedPtr<struct FLogEntryItem> SelectedItem, int64 UserData) {};
+};
+
 ENGINE_API  FArchive& operator<<(FArchive& Ar, FVisualLogDevice::FVisualLogEntryItem& FrameCacheItem);
 ENGINE_API  FArchive& operator<<(FArchive& Ar, FVisualLogDataBlock& Data);
 ENGINE_API  FArchive& operator<<(FArchive& Ar, FVisualLogHistogramSample& Sample);
