Fixed CIS compilation for OSX.

diff --git a/Engine/Source/Developer/LogVisualizer/Public/LogVisualizerSettings.h b/Engine/Source/Developer/LogVisualizer/Public/LogVisualizerSettings.h
--- a/Engine/Source/Developer/LogVisualizer/Public/LogVisualizerSettings.h
+++ b/Engine/Source/Developer/LogVisualizer/Public/LogVisualizerSettings.h
@@ -5,7 +5,6 @@ GameplayDebuggerSettings.h: Declares the UGameplayDebuggerSettings class.
 =============================================================================*/
 #pragma once
 
-#include "VisualLoggerDatabase.h"
 #include "LogVisualizerSettings.generated.h"
 
 DECLARE_MULTICAST_DELEGATE_TwoParams(FOnFilterCategoryAdded, FString, ELogVerbosity::Type);
