Next CIS compilation fix.

diff --git a/Engine/Source/Developer/LogVisualizer/Public/LogVisualizerSettings.h b/Engine/Source/Developer/LogVisualizer/Public/LogVisualizerSettings.h
--- a/Engine/Source/Developer/LogVisualizer/Public/LogVisualizerSettings.h
+++ b/Engine/Source/Developer/LogVisualizer/Public/LogVisualizerSettings.h
@@ -10,6 +10,8 @@ GameplayDebuggerSettings.h: Declares the UGameplayDebuggerSettings class.
 DECLARE_MULTICAST_DELEGATE_TwoParams(FOnFilterCategoryAdded, FString, ELogVerbosity::Type);
 DECLARE_MULTICAST_DELEGATE_OneParam(FOnFilterCategoryRemoved, FString);
 
+struct FVisualLoggerDBRow;
+
 USTRUCT()
 struct FCategoryFilter
 {
