CIS compilation fix.

diff --git a/Engine/Source/Developer/LogVisualizer/Private/SVisualLoggerFilters.cpp b/Engine/Source/Developer/LogVisualizer/Private/SVisualLoggerFilters.cpp
--- a/Engine/Source/Developer/LogVisualizer/Private/SVisualLoggerFilters.cpp
+++ b/Engine/Source/Developer/LogVisualizer/Private/SVisualLoggerFilters.cpp
@@ -98,7 +98,7 @@ bool SVisualLoggerFilters::GraphSubmenuVisibility(const FName MenuName)
 		return true;
 	}
 
-	if (CachedDatasPerGraph.Contains(MenuName) != NULL)
+	if (CachedDatasPerGraph.Contains(MenuName))
 	{
 		const TArray<FName>& DataNames = CachedDatasPerGraph[MenuName];
 		for (const FName& CurrentData : DataNames)
