Fixing compiler issues in LogVisualizer.

diff --git a/Engine/Source/Developer/LogVisualizer/Private/LogVisualizerSettings.cpp b/Engine/Source/Developer/LogVisualizer/Private/LogVisualizerSettings.cpp
--- a/Engine/Source/Developer/LogVisualizer/Private/LogVisualizerSettings.cpp
+++ b/Engine/Source/Developer/LogVisualizer/Private/LogVisualizerSettings.cpp
@@ -140,7 +140,7 @@ void FVisualLoggerFilters::AddCategory(FString InName, ELogVerbosity::Type InVer
 	Filter.CategoryName = InName;
 	Filter.LogVerbosity = InVerbosity;
 	Filter.Enabled = true;
-	int32 Index = Categories.Add(Filter);
+	Categories.Add(Filter);
 
 	FastCategoryFilterMap.Reset(); // we need to recreate cache - pointers can be broken
 	FastCategoryFilterMap.Reserve(Categories.Num());
diff --git a/Engine/Source/Developer/LogVisualizer/Private/SVisualLogger.cpp b/Engine/Source/Developer/LogVisualizer/Private/SVisualLogger.cpp
--- a/Engine/Source/Developer/LogVisualizer/Private/SVisualLogger.cpp
+++ b/Engine/Source/Developer/LogVisualizer/Private/SVisualLogger.cpp
@@ -97,6 +97,7 @@ SVisualLogger::SVisualLogger()
 	public:
 		FVisualLoggerDevice(SVisualLogger* InVisualLogger, FOnWorldChanged OnWorldChangedDelegate) : VisualLoggerWidget(InVisualLogger), LastUsedWorld(nullptr), OnWorldChanged(OnWorldChangedDelegate) {}
 
+		virtual ~FVisualLoggerDevice(){}
 		virtual void Serialize(const class UObject* LogOwner, FName OwnerName, FName OwnerClassName, const FVisualLogEntry& LogEntry) override
 		{
 			VisualLoggerWidget->OnNewLogEntry(FVisualLogDevice::FVisualLogEntryItem(OwnerName, OwnerClassName, LogEntry));
@@ -107,10 +108,10 @@ SVisualLogger::SVisualLogger()
 				LastUsedWorld = CurrentWorld;
 			}
 		}
-
+		
+		SVisualLogger* VisualLoggerWidget;
 		UWorld* LastUsedWorld;
 		FOnWorldChanged OnWorldChanged;
-		SVisualLogger* VisualLoggerWidget;
 	};
 
 	InternalDevice = MakeShareable(new FVisualLoggerDevice(this, FOnWorldChanged::CreateLambda([this](UWorld* PreviousWorld, UWorld* CurrentWorld){ OnNewWorld(CurrentWorld); })));
diff --git a/Engine/Source/Developer/LogVisualizer/Private/SVisualLoggerStatusView.cpp b/Engine/Source/Developer/LogVisualizer/Private/SVisualLoggerStatusView.cpp
--- a/Engine/Source/Developer/LogVisualizer/Private/SVisualLoggerStatusView.cpp
+++ b/Engine/Source/Developer/LogVisualizer/Private/SVisualLoggerStatusView.cpp
@@ -149,12 +149,12 @@ void SVisualLoggerStatusView::OnItemSelectionChanged(const FVisualLoggerDBRow& C
 	}
 
 	{
-		for (int32 ItemIndex = 0; ItemIndex < StatusItems.Num(); ItemIndex++)
+		for (int32 StatusItemIndex = 0; StatusItemIndex < StatusItems.Num(); StatusItemIndex++)
 		{
 			for (const FString& Category : ExpandedCategories)
-			{	if (StatusItems[ItemIndex]->ItemText == Category)
+			{	if (StatusItems[StatusItemIndex]->ItemText == Category)
 				{
-					StatusItemsView->SetItemExpansion(StatusItems[ItemIndex], true);
+					StatusItemsView->SetItemExpansion(StatusItems[StatusItemIndex], true);
 					break;
 				}
 			}
