Fix for Shadow variable issue in LogVisualizer.

diff --git a/Engine/Source/Developer/LogVisualizer/Private/SVisualLoggerView.cpp b/Engine/Source/Developer/LogVisualizer/Private/SVisualLoggerView.cpp
--- a/Engine/Source/Developer/LogVisualizer/Private/SVisualLoggerView.cpp
+++ b/Engine/Source/Developer/LogVisualizer/Private/SVisualLoggerView.cpp
@@ -286,10 +286,10 @@ void SVisualLoggerView::OnSearchSplitterResized()
 	SSplitter::FSlot const& LeftSplitterSlot = SearchSplitter->SlotAt(0);
 	SSplitter::FSlot const& RightSplitterSlot = SearchSplitter->SlotAt(1);
 
-	const float AnimationOutlinerFillPercentage = LeftSplitterSlot.SizeValue.Get() / RightSplitterSlot.SizeValue.Get();
-	SetAnimationOutlinerFillPercentage(AnimationOutlinerFillPercentage);
+	const float NewAnimationOutlinerFillPercentage = LeftSplitterSlot.SizeValue.Get() / RightSplitterSlot.SizeValue.Get();
+	SetAnimationOutlinerFillPercentage(NewAnimationOutlinerFillPercentage);
 
-	FLogVisualizer::Get().SetAnimationOutlinerFillPercentage(AnimationOutlinerFillPercentage);
+	FLogVisualizer::Get().SetAnimationOutlinerFillPercentage(NewAnimationOutlinerFillPercentage);
 }
 
 void SVisualLoggerView::OnSearchChanged(const FText& Filter)
