Fixes for monolitic build.

diff --git a/Engine/Source/Developer/LogVisualizer/Private/SVisualLogger.cpp b/Engine/Source/Developer/LogVisualizer/Private/SVisualLogger.cpp
--- a/Engine/Source/Developer/LogVisualizer/Private/SVisualLogger.cpp
+++ b/Engine/Source/Developer/LogVisualizer/Private/SVisualLogger.cpp
@@ -8,6 +8,7 @@
 #include "DesktopPlatformModule.h"
 #include "MainFrame.h"
 #include "VisualLoggerCameraController.h"
+#include "TimeSliderController.h"
 #if WITH_EDITOR
 #	include "Editor/UnrealEd/Public/EditorComponents.h"
 #	include "Editor/UnrealEd/Public/EditorReimportHandler.h"
diff --git a/Engine/Source/Developer/LogVisualizer/Private/VisualLoggerDatabase.cpp b/Engine/Source/Developer/LogVisualizer/Private/VisualLoggerDatabase.cpp
--- a/Engine/Source/Developer/LogVisualizer/Private/VisualLoggerDatabase.cpp
+++ b/Engine/Source/Developer/LogVisualizer/Private/VisualLoggerDatabase.cpp
@@ -3,7 +3,7 @@
 #include "LogVisualizer.h"
 #include "Engine.h"
 #include "VisualLoggerDatabase.h"
-
+#include "VisualLoggerRenderingActor.h"
 
 TSharedPtr< struct FVisualLoggerDatabase > FVisualLoggerDatabase::StaticInstance = nullptr;
 TSharedPtr< struct FVisualLoggerGraphsDatabase > FVisualLoggerGraphsDatabase::StaticInstance = nullptr;
