Fix non-unity build.

diff --git a/Engine/Source/Developer/LogVisualizer/Private/LogVisualizer.cpp b/Engine/Source/Developer/LogVisualizer/Private/LogVisualizer.cpp
--- a/Engine/Source/Developer/LogVisualizer/Private/LogVisualizer.cpp
+++ b/Engine/Source/Developer/LogVisualizer/Private/LogVisualizer.cpp
@@ -76,7 +76,7 @@ void FLogVisualizer::Shutdown()
 	StaticInstance.Reset();
 }
 
-inline void FLogVisualizer::Reset()
+void FLogVisualizer::Reset()
 {
 	TimeSliderController->SetTimesliderArgs(FVisualLoggerTimeSliderArgs());
 }
