Fixed crash in Visual Logger tool. It fixes UE-UE-20796 issue.
Fixed issue with missing debug geometry after PIE game.

diff --git a/Engine/Source/Developer/LogVisualizer/Private/SVisualLogger.cpp b/Engine/Source/Developer/LogVisualizer/Private/SVisualLogger.cpp
--- a/Engine/Source/Developer/LogVisualizer/Private/SVisualLogger.cpp
+++ b/Engine/Source/Developer/LogVisualizer/Private/SVisualLogger.cpp
@@ -65,25 +65,6 @@ namespace
 
 		return World;
 	}
-
-	static AVisualLoggerRenderingActor* GetVisualLoggerHelperActor(UWorld* World)
-	{
-		if (World == nullptr)
-		{
-			return nullptr;
-		}
-
-		for (TActorIterator<AVisualLoggerRenderingActor> It(World); It; ++It)
-		{
-			return *It;
-		}
-
-		FActorSpawnParameters SpawnInfo;
-		SpawnInfo.SpawnCollisionHandlingOverride = ESpawnActorCollisionHandlingMethod::AlwaysSpawn;
-		SpawnInfo.Name = *FString::Printf(TEXT("VisualLoggerRenderingActor"));
-		return World->SpawnActor<AVisualLoggerRenderingActor>(SpawnInfo);
-	}
-
 }
 
 SVisualLogger::SVisualLogger() 
@@ -329,7 +310,7 @@ void SVisualLogger::Construct(const FArguments& InArgs, const TSharedRef<SDockTa
 
 	DrawOnCanvasDelegateHandle = UDebugDrawService::Register(TEXT("VisLog"), FDebugDrawDelegate::CreateRaw(VisualLoggerCanvasRenderer.Get(), &FVisualLoggerCanvasRenderer::DrawOnCanvas));
 
-	AVisualLoggerRenderingActor* HelperActor = GetVisualLoggerHelperActor(LastUsedWorld.Get());
+	Cast<AVisualLoggerRenderingActor>(FVisualLoggerEditorInterface::Get()->GetHelperActor(LastUsedWorld.Get()));
 }
 
 void SVisualLogger::OnNewLogEntry(const FVisualLogDevice::FVisualLogEntryItem& Entry)
@@ -724,7 +705,7 @@ void SVisualLogger::ResetData()
 		VisualLoggerCanvasRenderer->ResetData();
 	}
 
-	if (AVisualLoggerRenderingActor* HelperActor = GetVisualLoggerHelperActor(LastUsedWorld.Get()))
+	if (AVisualLoggerRenderingActor* HelperActor = Cast<AVisualLoggerRenderingActor>(FVisualLoggerEditorInterface::Get()->GetHelperActor(LastUsedWorld.Get())))
 	{
 		HelperActor->ResetRendering();
 	}
@@ -755,7 +736,7 @@ void SVisualLogger::OnNewWorld(UWorld* NewWorld)
 	}
 	LastUsedWorld = NewWorld;
 
-	AVisualLoggerRenderingActor* HelperActor = GetVisualLoggerHelperActor(LastUsedWorld.Get());
+	AVisualLoggerRenderingActor* HelperActor = Cast<AVisualLoggerRenderingActor>(FVisualLoggerEditorInterface::Get()->GetHelperActor(LastUsedWorld.Get()));
 	check(HelperActor);
 
 	if (LastUsedWorld.IsValid() == false || ULogVisualizerSettings::StaticClass()->GetDefaultObject<ULogVisualizerSettings>()->bResetDataWithNewSession)
diff --git a/Engine/Source/Developer/LogVisualizer/Private/VisualLoggerDatabase.cpp b/Engine/Source/Developer/LogVisualizer/Private/VisualLoggerDatabase.cpp
--- a/Engine/Source/Developer/LogVisualizer/Private/VisualLoggerDatabase.cpp
+++ b/Engine/Source/Developer/LogVisualizer/Private/VisualLoggerDatabase.cpp
@@ -441,9 +441,9 @@ bool FVisualLoggerEditorInterface::IsItemVisible(FName RowName, int32 ItemIndex)
 	return FVisualLoggerDatabase::Get().GetRowByName(RowName).IsItemVisible(ItemIndex);
 }
 
-AActor* FVisualLoggerEditorInterface::GetHelperActor() const
+AActor* FVisualLoggerEditorInterface::GetHelperActor(UWorld* InWorld) const
 {
-	UWorld* World = GetWorld();
+	UWorld* World = InWorld ? InWorld : GetWorld();
 	if (World == nullptr)
 	{
 		return nullptr;
@@ -454,7 +454,10 @@ AActor* FVisualLoggerEditorInterface::GetHelperActor() const
 		return *It;
 	}
 
-	return nullptr;
+	FActorSpawnParameters SpawnInfo;
+	SpawnInfo.SpawnCollisionHandlingOverride = ESpawnActorCollisionHandlingMethod::AlwaysSpawn;
+	SpawnInfo.Name = *FString::Printf(TEXT("VisualLoggerRenderingActor"));
+	return World->SpawnActor<AVisualLoggerRenderingActor>(SpawnInfo);
 }
 
 bool FVisualLoggerEditorInterface::MatchCategoryFilters(const FString& String, ELogVerbosity::Type Verbosity)
diff --git a/Engine/Source/Developer/LogVisualizer/Private/VisualLoggerDatabase.h b/Engine/Source/Developer/LogVisualizer/Private/VisualLoggerDatabase.h
--- a/Engine/Source/Developer/LogVisualizer/Private/VisualLoggerDatabase.h
+++ b/Engine/Source/Developer/LogVisualizer/Private/VisualLoggerDatabase.h
@@ -201,7 +201,7 @@ struct FVisualLoggerEditorInterface : public IVisualLoggerEditorInterface
 	bool IsRowVisible(FName RowName) const override;
 	bool IsItemVisible(FName RowName, int32 ItemIndex) const override;
 	UWorld* GetWorld() const override;
-	AActor* GetHelperActor() const override;
+	AActor* GetHelperActor(UWorld* InWorld = nullptr) const override;
 
 	bool MatchCategoryFilters(const FString& String, ELogVerbosity::Type Verbosity = ELogVerbosity::All) override;
 };
diff --git a/Engine/Source/Runtime/Engine/Public/VisualLogger/VisualLoggerTypes.h b/Engine/Source/Runtime/Engine/Public/VisualLogger/VisualLoggerTypes.h
--- a/Engine/Source/Runtime/Engine/Public/VisualLogger/VisualLoggerTypes.h
+++ b/Engine/Source/Runtime/Engine/Public/VisualLogger/VisualLoggerTypes.h
@@ -280,7 +280,7 @@ struct IVisualLoggerEditorInterface
 	virtual bool IsRowVisible(FName RowName) const = 0;
 	virtual bool IsItemVisible(FName RowName, int32 ItemIndex) const = 0;
 	virtual UWorld* GetWorld() const = 0;
-	virtual AActor* GetHelperActor() const = 0;
+	virtual AActor* GetHelperActor(UWorld* InWorld = nullptr) const = 0;
 
 	virtual bool MatchCategoryFilters(const FString& String, ELogVerbosity::Type Verbosity = ELogVerbosity::All) = 0;
 };
