fixed triangulation error of some convex polys recorded by visual logger
#ue4

diff --git a/Engine/Source/Developer/LogVisualizer/Private/VisualLoggerRenderingActor.cpp b/Engine/Source/Developer/LogVisualizer/Private/VisualLoggerRenderingActor.cpp
--- a/Engine/Source/Developer/LogVisualizer/Private/VisualLoggerRenderingActor.cpp
+++ b/Engine/Source/Developer/LogVisualizer/Private/VisualLoggerRenderingActor.cpp
@@ -297,20 +297,50 @@ void AVisualLoggerRenderingActor::AddDebugRendering()
 
 namespace
 {
+	static bool IsPolygonWindingCorrect(const TArray<FVector>& Verts)
+	{
+		// this will work only for convex polys, but we're assuming that all logged polygons are convex in the first place
+		if (Verts.Num() >= 3)
+		{
+			const FVector SurfaceNormal = FVector::CrossProduct(Verts[1] - Verts[0], Verts[2] - Verts[0]);
+			const float TestDot = FVector::DotProduct(SurfaceNormal, FVector(0, 0, 1));
+			return TestDot > 0;
+		}
+
+		return false;
+	}
+
 	static void GetPolygonMesh(const FVisualLogShapeElement* ElementToDraw, FDebugRenderSceneProxy::FMesh& TestMesh, const FVector& VertexOffset = FVector::ZeroVector)
 	{
 		TestMesh.Color = ElementToDraw->GetFColor();
 
 		FClipSMPolygon InPoly(ElementToDraw->Points.Num());
-		for (int32 Index = 0; Index < ElementToDraw->Points.Num(); Index++)
+		InPoly.FaceNormal = FVector(0, 0, 1);
+
+		const bool bHasCorrectWinding = IsPolygonWindingCorrect(ElementToDraw->Points);
+		if (bHasCorrectWinding)
 		{
-			FClipSMVertex v1;
-			v1.Pos = ElementToDraw->Points[Index];
-			InPoly.Vertices.Add(v1);
+			for (int32 Index = 0; Index < ElementToDraw->Points.Num(); Index++)
+			{
+				FClipSMVertex v1;
+				v1.Pos = ElementToDraw->Points[Index];
+				InPoly.Vertices.Add(v1);
+			}
+		}
+		else
+		{
+			for (int32 Index = ElementToDraw->Points.Num() - 1; Index >= 0; Index--)
+			{
+				FClipSMVertex v1;
+				v1.Pos = ElementToDraw->Points[Index];
+				InPoly.Vertices.Add(v1);
+			}
 		}
 
 		TArray<FClipSMTriangle> OutTris;
-		if (TriangulatePoly(OutTris, InPoly, false))
+		
+		const bool bTriangulated = TriangulatePoly(OutTris, InPoly, false);
+		if (bTriangulated)
 		{
 			int32 LastIndex = 0;
 
