Fix shadow variable and bug caused by it

diff --git a/Engine/Source/Developer/LogVisualizer/Private/STimelinesContainer.cpp b/Engine/Source/Developer/LogVisualizer/Private/STimelinesContainer.cpp
--- a/Engine/Source/Developer/LogVisualizer/Private/STimelinesContainer.cpp
+++ b/Engine/Source/Developer/LogVisualizer/Private/STimelinesContainer.cpp
@@ -77,17 +77,17 @@ bool STimelinesContainer::IsNodeSelected(TSharedPtr<STimeline> Node) const
 	return FVisualLoggerDatabase::Get().IsRowSelected(Node->GetName());// SelectedNodes.Contains(Node);
 }
 
-void STimelinesContainer::ChangeSelection(TSharedPtr<STimeline> Timeline, const FPointerEvent& MouseEvent)
+void STimelinesContainer::ChangeSelection(TSharedPtr<STimeline> InTimeline, const FPointerEvent& MouseEvent)
 {
 	if (MouseEvent.IsLeftShiftDown() == false)
 	{
 		if (MouseEvent.IsLeftControlDown())
 		{
-			SetSelectionState(Timeline, !Timeline->IsSelected(), false);
+			SetSelectionState(InTimeline, !InTimeline->IsSelected(), false);
 		}
 		else
 		{
-			SetSelectionState(Timeline, true, true);
+			SetSelectionState(InTimeline, true, true);
 		}
 	}
 	else
@@ -101,9 +101,9 @@ void STimelinesContainer::ChangeSelection(TSharedPtr<STimeline> Timeline, const
 		if (LastSelected.IsValid())
 		{
 			bool bStartedSelection = false;
-			for (TSharedPtr<STimeline>& Timeline : TimelineItems)
+			for (TSharedPtr<STimeline>& TimelineItem : TimelineItems)
 			{
-				if (Timeline == LastSelected || Timeline == Timeline)
+				if (TimelineItem == LastSelected || InTimeline == TimelineItem)
 				{
 					if (!bStartedSelection)
 					{
@@ -117,11 +117,11 @@ void STimelinesContainer::ChangeSelection(TSharedPtr<STimeline> Timeline, const
 				}
 				if (bStartedSelection)
 				{
-					SetSelectionState(Timeline, true, false);
+					SetSelectionState(TimelineItem, true, false);
 				}
 			}
 		}
-		SetSelectionState(Timeline, true, false);
+		SetSelectionState(InTimeline, true, false);
 	}
 }
 
