Copying //UE4/Dev-Sequencer to //UE4/Main

diff --git a/Engine/Plugins/Editor/SequenceRecorder/Source/SequenceRecorder/Private/SequenceRecorderSettings.h b/Engine/Plugins/Editor/SequenceRecorder/Source/SequenceRecorder/Private/SequenceRecorderSettings.h
--- a/Engine/Plugins/Editor/SequenceRecorder/Source/SequenceRecorder/Private/SequenceRecorderSettings.h
+++ b/Engine/Plugins/Editor/SequenceRecorder/Source/SequenceRecorder/Private/SequenceRecorderSettings.h
@@ -2,6 +2,7 @@
 
 #pragma once
 
+#include "SequenceRecorderActorFilter.h"
 #include "SequenceRecorderSettings.generated.h"
 
 UCLASS(config=Editor, MinimalAPI)
@@ -34,7 +35,19 @@ public:
 	UPROPERTY(EditAnywhere, AdvancedDisplay, Category = "Sequence Recording")
 	FString AnimationSubDirectory;
 
-	/** When set, record a transform track and record animations in local space */
+	/** Whether to record nearby spawned actors. */
 	UPROPERTY(EditAnywhere, Category = "Sequence Recording")
-	bool bRecordTransformsAsSeperateTrack;
+	bool bRecordNearbySpawnedActors;
+
+	/** Proximity to currently recorded actors to record newly spawned actors. */
+	UPROPERTY(EditAnywhere, Category = "Sequence Recording")
+	float NearbyActorRecordingProximity;
+
+	/** Whether to record the world settings actor in the sequence (some games use this to attach world SFX) */
+	UPROPERTY(EditAnywhere, Category = "Sequence Recording")
+	bool bRecordWorldSettingsActor;
+
+	/** Filter to check spawned actors against to see if they should be recorded */
+	UPROPERTY(EditAnywhere, Category = "Sequence Recording")
+	FSequenceRecorderActorFilter ActorFilter;
 };
\ No newline at end of file
diff --git a/Engine/Plugins/Editor/SequenceRecorder/Source/SequenceRecorder/Private/SequenceRecorderUtils.cpp b/Engine/Plugins/Editor/SequenceRecorder/Source/SequenceRecorder/Private/SequenceRecorderUtils.cpp
new file mode 100644
--- /dev/null
+++ b/Engine/Plugins/Editor/SequenceRecorder/Source/SequenceRecorder/Private/SequenceRecorderUtils.cpp
@@ -0,0 +1,26 @@
+// Copyright 1998-2016 Epic Games, Inc. All Rights Reserved.
+
+#include "SequenceRecorderPrivatePCH.h"
+#include "SequenceRecorderUtils.h"
+
+namespace SequenceRecorderUtils
+{
+
+AActor* GetAttachment(AActor* InActor, FName& SocketName, FName& ComponentName)
+{
+	AActor* AttachedActor = nullptr;
+	if(InActor)
+	{
+		USceneComponent* RootComponent = InActor->GetRootComponent();
+		if(RootComponent && RootComponent->AttachParent != nullptr)
+		{
+			AttachedActor = RootComponent->AttachParent->GetOwner();
+			SocketName = RootComponent->AttachSocketName;
+			ComponentName = RootComponent->AttachParent->GetFName();
+		}
+	}
+
+	return AttachedActor;
+}
+
+}
\ No newline at end of file
diff --git a/Engine/Plugins/Editor/SequenceRecorder/Source/SequenceRecorder/Private/SequenceRecorderUtils.h b/Engine/Plugins/Editor/SequenceRecorder/Source/SequenceRecorder/Private/SequenceRecorderUtils.h
--- a/Engine/Plugins/Editor/SequenceRecorder/Source/SequenceRecorder/Private/SequenceRecorderUtils.h
+++ b/Engine/Plugins/Editor/SequenceRecorder/Source/SequenceRecorder/Private/SequenceRecorderUtils.h
@@ -37,4 +37,7 @@ static AssetType* MakeNewAsset(const FString& BaseAssetPath, const FString& Base
 	return NewObject<AssetType>(Package, *AssetName, RF_Public | RF_Standalone);
 }
 
+/** Helper function - check whether our component hierarchy has some attachment outside of its owned components */
+AActor* GetAttachment(AActor* InActor, FName& SocketName, FName& ComponentName);
+
 }
\ No newline at end of file
diff --git a/Engine/Plugins/Editor/SequenceRecorder/Source/SequenceRecorder/SequenceRecorder.Build.cs b/Engine/Plugins/Editor/SequenceRecorder/Source/SequenceRecorder/SequenceRecorder.Build.cs
--- a/Engine/Plugins/Editor/SequenceRecorder/Source/SequenceRecorder/SequenceRecorder.Build.cs
+++ b/Engine/Plugins/Editor/SequenceRecorder/Source/SequenceRecorder/SequenceRecorder.Build.cs
@@ -43,6 +43,7 @@ namespace UnrealBuildTool.Rules
 					"MovieScene",
 					"MovieSceneTracks",
 					"LevelSequence",
+					"NetworkReplayStreaming",
 				}
 				);
 
diff --git a/Engine/Plugins/MovieScene/LevelSequenceEditor/Source/LevelSequenceEditor/LevelSequenceEditor.Build.cs b/Engine/Plugins/MovieScene/LevelSequenceEditor/Source/LevelSequenceEditor/LevelSequenceEditor.Build.cs
--- a/Engine/Plugins/MovieScene/LevelSequenceEditor/Source/LevelSequenceEditor/LevelSequenceEditor.Build.cs
+++ b/Engine/Plugins/MovieScene/LevelSequenceEditor/Source/LevelSequenceEditor/LevelSequenceEditor.Build.cs
@@ -19,6 +19,7 @@ public class LevelSequenceEditor : ModuleRules
 				"AppFramework",
                 "LevelSequence",
 				"BlueprintGraph",
+                "CinematicCamera",
 				"Core",
 				"CoreUObject",
                 "EditorStyle",
@@ -42,6 +43,7 @@ public class LevelSequenceEditor : ModuleRules
                 "MovieSceneTools",
 				"SceneOutliner",
 				"PlacementMode",
+                "Settings",
 			}
 		);
 
diff --git a/Engine/Plugins/MovieScene/LevelSequenceEditor/Source/LevelSequenceEditor/Private/CinematicViewport/FilmOverlays.cpp b/Engine/Plugins/MovieScene/LevelSequenceEditor/Source/LevelSequenceEditor/Private/CinematicViewport/FilmOverlays.cpp
--- a/Engine/Plugins/MovieScene/LevelSequenceEditor/Source/LevelSequenceEditor/Private/CinematicViewport/FilmOverlays.cpp
+++ b/Engine/Plugins/MovieScene/LevelSequenceEditor/Source/LevelSequenceEditor/Private/CinematicViewport/FilmOverlays.cpp
@@ -256,26 +256,13 @@ struct FFilmOverlay_SafeFrame : IFilmOverlay
 		auto SetValue = [this](float NewValue){ SizePercentage = NewValue; };
 		auto SetValueCommitted = [this](float NewValue, ETextCommit::Type){ SizePercentage = NewValue; };
 
-		return SNew(SHorizontalBox)
-
-			// + SHorizontalBox::Slot()
-			// .AutoWidth()
-			// [
-			// 	SNew(STextBlock)
-			// 	.Text(LOCTEXT("SizePercentage", "Screen Percentage:"))
-			// ]
-
-			+ SHorizontalBox::Slot()
-			.HAlign(HAlign_Right)
-			[
-				SNew(SSpinBox<float>)
-				.MinValue(1)
-				.MaxValue(99)
-				.Value_Lambda(GetValue)
-				.OnValueChanged_Lambda(SetValue)
-				.OnValueCommitted_Lambda(SetValueCommitted)
-				.TypeInterface(MakeShareable(new FPercentageInterface))
-			];
+		return SNew(SSpinBox<float>)
+			.MinValue(1)
+			.MaxValue(99)
+			.Value_Lambda(GetValue)
+			.OnValueChanged_Lambda(SetValue)
+			.OnValueCommitted_Lambda(SetValueCommitted)
+			.TypeInterface(MakeShareable(new FPercentageInterface));
 	}
 
 private:
@@ -532,7 +519,11 @@ TSharedRef<SWidget> SFilmOverlayOptions::ConstructToggleableOverlaysMenu()
 		.HAlign(HAlign_Right)
 		.Padding(10.f, 5.f)
 		[
-			Settings.ToSharedRef()
+			SNew(SBox)
+			.WidthOverride(100)
+			[
+				Settings.ToSharedRef()
+			]
 		];
 
 		++Row;
diff --git a/Engine/Plugins/MovieScene/LevelSequenceEditor/Source/LevelSequenceEditor/Private/CinematicViewport/SCinematicLevelViewport.cpp b/Engine/Plugins/MovieScene/LevelSequenceEditor/Source/LevelSequenceEditor/Private/CinematicViewport/SCinematicLevelViewport.cpp
--- a/Engine/Plugins/MovieScene/LevelSequenceEditor/Source/LevelSequenceEditor/Private/CinematicViewport/SCinematicLevelViewport.cpp
+++ b/Engine/Plugins/MovieScene/LevelSequenceEditor/Source/LevelSequenceEditor/Private/CinematicViewport/SCinematicLevelViewport.cpp
@@ -9,13 +9,14 @@
 #include "MovieSceneCinematicShotSection.h"
 #include "MovieSceneSubTrack.h"
 #include "MovieSceneSubSection.h"
+#include "MovieSceneCommonHelpers.h"
 #include "SCinematicTransportRange.h"
 #include "FilmOverlays.h"
 #include "LevelViewportLayout.h"
 #include "LevelViewportTabContent.h"
 #include "LevelSequenceEditorStyle.h"
 #include "SWidgetSwitcher.h"
-#include "Camera/CineCameraComponent.h"
+#include "CineCameraComponent.h"
 #include "UnitConversion.h"
 #include "ScopedTransaction.h"
 
@@ -79,6 +80,21 @@ FCinematicViewportClient::FCinematicViewportClient()
 	bDisableInput = false;
 }
 
+bool FCinematicViewportClient::InputKey(FViewport* Viewport, int32 ControllerId, FKey Key, EInputEvent Event, float AmountDepressed, bool bGamepad)
+{
+	// Allow immersive mode to be processed when input handling is disabled
+	if (bDisableInput)		
+	{
+		if (Key == EKeys::F11)
+		{
+			return false;
+		}
+	}
+
+	return FLevelEditorViewportClient::InputKey(Viewport, ControllerId, Key, Event, AmountDepressed, bGamepad);
+}
+
+
 class SPreArrangedBox : public SCompoundWidget
 {
 public:
@@ -428,6 +444,11 @@ void SCinematicLevelViewport::Construct(const FArguments& InArgs)
 		Setup(Toolkit);
 		return false;
 	});
+
+	CommandList = MakeShareable( new FUICommandList );
+	// Ensure the commands are registered
+	FLevelSequenceEditorCommands::Register();
+	BindCommands();
 }
 END_SLATE_FUNCTION_BUILD_OPTIMIZATION
 
@@ -481,6 +502,29 @@ bool SCinematicLevelViewport::IsMaximized() const
 	return false;
 }
 
+void SCinematicLevelViewport::OnToggleImmersive()
+{
+	TSharedPtr<FLevelViewportLayout> ParentLayoutPinned = ParentLayout.Pin();
+	if (ParentLayoutPinned.IsValid() && ParentLayoutPinned->IsMaximizeSupported())
+	{
+		const bool bAllowAnimation = true;
+		const bool bMaximize = !ParentLayoutPinned->IsViewportMaximized(LayoutName);
+		
+		const bool bImmersive = !IsImmersive();
+		ParentLayoutPinned->RequestMaximizeViewport( LayoutName, bMaximize, bImmersive, bAllowAnimation );
+	}
+}
+
+bool SCinematicLevelViewport::IsImmersive() const
+{
+	TSharedPtr<FLevelViewportLayout> ParentLayoutPinned = ParentLayout.Pin();
+	if (ParentLayoutPinned.IsValid())
+	{
+		return ParentLayoutPinned->IsViewportImmersive(LayoutName);
+	}
+	return false;
+}
+
 const FSlateBrush* SCinematicLevelViewport::GetMaximizeImage() const
 {
 	return IsMaximized() ? FEditorStyle::GetBrush("LevelViewportToolBar.Maximize.Checked") : FEditorStyle::GetBrush("LevelViewportToolBar.Maximize.Normal");
@@ -708,10 +752,17 @@ void SCinematicLevelViewport::UpdateActiveCameras()
 	for (int32 SpawnableIndex = 0; SpawnableIndex < NumSpawnables; ++SpawnableIndex)
 	{
 		FMovieSceneSpawnable& Spawnable = MovieScene->GetSpawnable(SpawnableIndex);
-		if (Spawnable.GetClass()->IsChildOf(ACameraActor::StaticClass()))
+		if (Spawnable.GetClass()->IsChildOf(AActor::StaticClass()))
 		{
-			ACameraActor* Camera = Cast<ACameraActor>(SequenceInstance->FindObject(Spawnable.GetGuid(), *Sequencer));
-			Cameras.Add(Spawnable.GetGuid(), FActiveCamera(Spawnable.GetName(), Camera));
+			AActor* CameraActor = Cast<AActor>(SequenceInstance->FindObject(Spawnable.GetGuid(), *Sequencer));
+			if (CameraActor)
+			{
+				UCameraComponent* CameraComponent = MovieSceneHelpers::CameraComponentFromActor(CameraActor);
+				if (CameraComponent)
+				{
+					Cameras.Add(Spawnable.GetGuid(), FActiveCamera(Spawnable.GetName(), CameraActor));
+				}
+			}
 		}
 	}
 
@@ -719,10 +770,14 @@ void SCinematicLevelViewport::UpdateActiveCameras()
 	for (int32 PossessableIndex = 0; PossessableIndex < NumPossessables; ++PossessableIndex)
 	{
 		FMovieScenePossessable& Possessable = MovieScene->GetPossessable(PossessableIndex);
-		ACameraActor* Camera = Cast<ACameraActor>(SequenceInstance->FindObject(Possessable.GetGuid(), *Sequencer));
-		if (Camera)
+		AActor* CameraActor = Cast<AActor>(SequenceInstance->FindObject(Possessable.GetGuid(), *Sequencer));
+		if (CameraActor)
 		{
-			Cameras.Add(Possessable.GetGuid(), FActiveCamera(Possessable.GetName(), Camera));
+			UCameraComponent* CameraComponent = MovieSceneHelpers::CameraComponentFromActor(CameraActor);
+			if (CameraComponent)
+			{
+				Cameras.Add(Possessable.GetGuid(), FActiveCamera(Possessable.GetName(), CameraActor));
+			}
 		}
 	}
 }
@@ -739,7 +794,7 @@ void SCinematicLevelViewport::OnUpdateCameraCut(UObject* CameraObject)
 		return;
 	}
 
-	ACameraActor* CameraActor = Cast<ACameraActor>(CameraObject);
+	AActor* CameraActor = Cast<AActor>(CameraObject);
 	TSharedPtr<FCinematicViewportClient> ViewportClient = GetViewportClient();
 	
 	if (CameraActor)
@@ -763,17 +818,21 @@ void SCinematicLevelViewport::OnUpdateCameraCut(UObject* CameraObject)
 	// If viewing through a camera - enforce aspect ratio.
 	if (CameraActor)
 	{
-		if (CameraActor->GetCameraComponent()->AspectRatio == 0)
+		UCameraComponent* CameraComponent = MovieSceneHelpers::CameraComponentFromActor(CameraActor);
+		if (CameraComponent)
 		{
-			ViewportClient->AspectRatio = 1.7f;
-		}
-		else
-		{
-			ViewportClient->AspectRatio = CameraActor->GetCameraComponent()->AspectRatio;
-		}
+			if (CameraComponent->AspectRatio == 0)
+			{
+				ViewportClient->AspectRatio = 1.7f;
+			}
+			else
+			{
+				ViewportClient->AspectRatio = CameraComponent->AspectRatio;
+			}
 
-		//don't stop the camera from zooming when not playing back
-		ViewportClient->ViewFOV = CameraActor->GetCameraComponent()->FieldOfView;
+			//don't stop the camera from zooming when not playing back
+			ViewportClient->ViewFOV = CameraComponent->FieldOfView;
+		}
 	}
 
 	// Update ControllingActorViewInfo, so it is in sync with the updated viewport
@@ -782,6 +841,11 @@ void SCinematicLevelViewport::OnUpdateCameraCut(UObject* CameraObject)
 
 FReply SCinematicLevelViewport::OnKeyDown( const FGeometry& MyGeometry, const FKeyEvent& InKeyEvent ) 
 {
+	if (CommandList->ProcessCommandBindings(InKeyEvent))
+	{
+		return FReply::Handled();
+	}
+
 	ISequencer* Sequencer = GetSequencer();
 	if (Sequencer && Sequencer->GetCommandBindings()->ProcessCommandBindings(InKeyEvent))
 	{
@@ -861,6 +925,19 @@ void SCinematicLevelViewport::OnEditorClosed()
 	}
 }
 
+void SCinematicLevelViewport::BindCommands()
+{
+	FUICommandList& CommandListRef = *CommandList;
+
+	const FLevelSequenceEditorCommands& Commands = FLevelSequenceEditorCommands::Get();
+
+	CommandListRef.MapAction( 
+		Commands.ToggleImmersive,
+		FExecuteAction::CreateSP( this, &SCinematicLevelViewport::OnToggleImmersive ),
+		FCanExecuteAction(),
+		FIsActionChecked::CreateSP(this, &SCinematicLevelViewport::IsImmersive));
+}
+
 ISequencer* SCinematicLevelViewport::GetSequencer() const
 {
 	TSharedPtr<FLevelSequenceEditorToolkit> Toolkit = CurrentToolkit.Pin();
diff --git a/Engine/Plugins/MovieScene/LevelSequenceEditor/Source/LevelSequenceEditor/Private/CinematicViewport/SCinematicLevelViewport.h b/Engine/Plugins/MovieScene/LevelSequenceEditor/Source/LevelSequenceEditor/Private/CinematicViewport/SCinematicLevelViewport.h
--- a/Engine/Plugins/MovieScene/LevelSequenceEditor/Source/LevelSequenceEditor/Private/CinematicViewport/SCinematicLevelViewport.h
+++ b/Engine/Plugins/MovieScene/LevelSequenceEditor/Source/LevelSequenceEditor/Private/CinematicViewport/SCinematicLevelViewport.h
@@ -18,6 +18,7 @@ struct FCinematicViewportClient : FLevelEditorViewportClient
 {
 	FCinematicViewportClient();
 	virtual bool IsLevelEditorClient() const override { return true; }
+	virtual bool InputKey(FViewport* Viewport, int32 ControllerId, FKey Key, EInputEvent Event, float AmountDepressed = 1.f, bool bGamepad=false) override;
 };
 
 
@@ -87,6 +88,8 @@ private:
 	/** Called when the level sequence editor toolkit we were observing has been closed */
 	void OnEditorClosed();
 
+	void BindCommands();
+
 private:
 
 	/** Get the sequencer ptr from our current toolkit, if set */
@@ -109,6 +112,12 @@ private:
 	/** Test whether this viewport is maximized or not */
 	bool IsMaximized() const;
 
+	/** Toggle this viewport's immersive state */
+	void OnToggleImmersive();
+
+	/** Test whether this viewport is immersive or not */
+	bool IsImmersive() const;
+
 	/** Called when we want to revert to the equivalent non-cinematic viewport */
 	FReply OnRevertLayout();
 
@@ -154,6 +163,9 @@ private:
 	/** The toolkit we're currently editing */
 	TWeakPtr<FLevelSequenceEditorToolkit> CurrentToolkit;
 
+	/** Commandlist used in the viewport (Maps commands to viewport specific actions) */
+	TSharedPtr<FUICommandList> CommandList;
+
 	/** Slot for transport controls */
 	TSharedPtr<SBox> TransportControls;
 
@@ -192,12 +204,12 @@ private:
 	/** Cached active cameras */
 	struct FActiveCamera
 	{
-		FActiveCamera(FString InDisplayName, ACameraActor* InCameraActor)
+		FActiveCamera(FString InDisplayName, AActor* InCameraActor)
 			: DisplayName(InDisplayName), CameraActor(InCameraActor)
 		{}
 
 		FString DisplayName;
-		TWeakObjectPtr<ACameraActor> CameraActor;
+		TWeakObjectPtr<AActor> CameraActor;
 	};
 	TMap<FGuid, FActiveCamera> Cameras;
 };
\ No newline at end of file
diff --git a/Engine/Plugins/MovieScene/LevelSequenceEditor/Source/LevelSequenceEditor/Private/LevelSequenceEditorCommands.cpp b/Engine/Plugins/MovieScene/LevelSequenceEditor/Source/LevelSequenceEditor/Private/LevelSequenceEditorCommands.cpp
new file mode 100644
--- /dev/null
+++ b/Engine/Plugins/MovieScene/LevelSequenceEditor/Source/LevelSequenceEditor/Private/LevelSequenceEditorCommands.cpp
@@ -0,0 +1,19 @@
+// Copyright 1998-2016 Epic Games, Inc. All Rights Reserved.
+
+#include "LevelSequenceEditorPCH.h"
+#include "LevelSequenceEditorCommands.h"
+
+#define LOCTEXT_NAMESPACE "LevelSequenceEditorCommands"
+
+FLevelSequenceEditorCommands::FLevelSequenceEditorCommands()
+	: TCommands<FLevelSequenceEditorCommands>("LevelSequenceEditor", LOCTEXT("LevelSequenceEditor", "Level Sequence Editor"), NAME_None, "LevelSequenceEditorStyle")
+{
+}
+
+void FLevelSequenceEditorCommands::RegisterCommands()
+{
+	UI_COMMAND(CreateNewLevelSequenceInLevel, "Add Level Sequence", "Create a new level sequence asset, and place an instance of it in this level", EUserInterfaceActionType::Button, FInputChord());
+	UI_COMMAND(ToggleImmersive, "Immersive Mode", "Switches this cinematic viewport between immersive mode and regular mode", EUserInterfaceActionType::ToggleButton, FInputChord( EKeys::F11 ) );
+}
+
+#undef LOCTEXT_NAMESPACE
\ No newline at end of file
diff --git a/Engine/Plugins/MovieScene/LevelSequenceEditor/Source/LevelSequenceEditor/Private/LevelSequenceEditorModule.cpp b/Engine/Plugins/MovieScene/LevelSequenceEditor/Source/LevelSequenceEditor/Private/LevelSequenceEditorModule.cpp
--- a/Engine/Plugins/MovieScene/LevelSequenceEditor/Source/LevelSequenceEditor/Private/LevelSequenceEditorModule.cpp
+++ b/Engine/Plugins/MovieScene/LevelSequenceEditor/Source/LevelSequenceEditor/Private/LevelSequenceEditorModule.cpp
@@ -1,37 +1,21 @@
 // Copyright 1998-2016 Epic Games, Inc. All Rights Reserved.
 
 #include "LevelSequenceEditorPCH.h"
+#include "CineCameraActor.h"
+#include "CinematicLevelViewportLayout.h"
+#include "IPlacementModeModule.h"
+#include "ISettingsModule.h"
 #include "LevelEditor.h"
 #include "LevelSequenceActor.h"
 #include "LevelSequenceEditorStyle.h"
 #include "ModuleInterface.h"
 #include "PropertyEditorModule.h"
-#include "CinematicLevelViewportLayout.h"
-#include "IPlacementModeModule.h"
-#include "Camera/CineCameraActor.h"
 
+
 #define LOCTEXT_NAMESPACE "LevelSequenceEditor"
 
+
 TSharedPtr<FLevelSequenceEditorStyle> FLevelSequenceEditorStyle::Singleton;
 
-class FLevelSequenceExtensionCommands
-	: public TCommands<FLevelSequenceExtensionCommands>
-{
-public:
 
-	/** Default constructor. */
-	FLevelSequenceExtensionCommands()
-		: TCommands("LevelSequenceEditor", LOCTEXT("ExtensionDescription", "Extension commands specific to the level sequence editor"), NAME_None, "LevelSequenceEditorStyle")
-	{ }
-
-	/** Initialize commands */
-	virtual void RegisterCommands() override
-	{
-		UI_COMMAND(CreateNewLevelSequenceInLevel, "Add Level Sequence", "Create a new level sequence asset, and place an instance of it in this level", EUserInterfaceActionType::Button, FInputChord());
-	}
-
-	TSharedPtr<FUICommandInfo> CreateNewLevelSequenceInLevel;
-};
-
-
 /**
@@ -53,6 +37,7 @@
 		RegisterMenuExtensions();
 		RegisterLevelEditorExtensions();
 		RegisterPlacementModeExtensions();
+		RegisterSettings();
 	}
 	
 	virtual void ShutdownModule() override
@@ -62,10 +47,38 @@
 		UnregisterMenuExtensions();
 		UnregisterLevelEditorExtensions();
 		UnregisterPlacementModeExtensions();
+		UnregisterSettings();
 	}
 
 protected:
 
+	/** Registers asset tool actions. */
+	void RegisterAssetTools()
+	{
+		IAssetTools& AssetTools = FModuleManager::LoadModuleChecked<FAssetToolsModule>("AssetTools").Get();
+		RegisterAssetTypeAction(AssetTools, MakeShareable(new FLevelSequenceActions(FLevelSequenceEditorStyle::Get())));
+	}
+
+	/**
+	* Registers a single asset type action.
+	*
+	* @param AssetTools The asset tools object to register with.
+	* @param Action The asset type action to register.
+	*/
+	void RegisterAssetTypeAction(IAssetTools& AssetTools, TSharedRef<IAssetTypeActions> Action)
+	{
+		AssetTools.RegisterAssetTypeActions(Action);
+		RegisteredAssetTypeActions.Add(Action);
+	}
+
+	/** Register details view customizations. */
+	void RegisterCustomizations()
+	{
+		FPropertyEditorModule& PropertyModule = FModuleManager::LoadModuleChecked<FPropertyEditorModule>("PropertyEditor");
+		LevelSequencePlayingSettingsName = FLevelSequencePlaybackSettings::StaticStruct()->GetFName();
+		PropertyModule.RegisterCustomPropertyTypeLayout(LevelSequencePlayingSettingsName, FOnGetPropertyTypeCustomizationInstance::CreateStatic(&FLevelSequencePlaybackSettingsCustomization::MakeInstance));
+	}
+
 	/** Registers level editor extensions. */
 	void RegisterLevelEditorExtensions()
 	{
@@ -84,11 +97,24 @@
 		LevelEditorModule.RegisterCustomViewportLayout("TwoPaneCinematic", OnePaneDefinition);
 	}
 
-	/** Registers asset tool actions. */
-	void RegisterAssetTools()
+	/** Register menu extensions for the level editor toolbar. */
+	void RegisterMenuExtensions()
 	{
-		IAssetTools& AssetTools = FModuleManager::LoadModuleChecked<FAssetToolsModule>("AssetTools").Get();
-		RegisterAssetTypeAction(AssetTools, MakeShareable(new FLevelSequenceActions(FLevelSequenceEditorStyle::Get())));
+		FLevelSequenceEditorCommands::Register();
+
+		CommandList = MakeShareable(new FUICommandList);
+		CommandList->MapAction(FLevelSequenceEditorCommands::Get().CreateNewLevelSequenceInLevel,
+			FExecuteAction::CreateStatic(&FLevelSequenceEditorModule::OnCreateActorInLevel)
+		);
+
+		// Create and register the level editor toolbar menu extension
+		CinematicsMenuExtender = MakeShareable(new FExtender);
+		CinematicsMenuExtender->AddMenuExtension("LevelEditorNewMatinee", EExtensionHook::First, CommandList, FMenuExtensionDelegate::CreateStatic([](FMenuBuilder& MenuBuilder) {
+			MenuBuilder.AddMenuEntry(FLevelSequenceEditorCommands::Get().CreateNewLevelSequenceInLevel);
+		}));
+
+		FLevelEditorModule& LevelEditorModule = FModuleManager::LoadModuleChecked<FLevelEditorModule>("LevelEditor");
+		LevelEditorModule.GetAllLevelEditorToolbarCinematicsMenuExtenders().Add(CinematicsMenuExtender);
 	}
 
 	/** Registers placement mode extensions. */
@@ -98,23 +124,31 @@
 			LOCTEXT("CinematicCategoryName", "Cinematic"),
 			"Cinematic",
 			TEXT("PMCinematic"),
-			25);
+			25
+		);
+
 		IPlacementModeModule::Get().RegisterPlacementCategory(Info);
 		IPlacementModeModule::Get().RegisterPlaceableItem(Info.UniqueHandle, MakeShareable( new FPlaceableItem(nullptr, FAssetData(ACineCameraActor::StaticClass())) ));
 	}
 
-	/**
-	 * Registers a single asset type action.
-	 *
-	 * @param AssetTools The asset tools object to register with.
-	 * @param Action The asset type action to register.
-	 */
-	void RegisterAssetTypeAction(IAssetTools& AssetTools, TSharedRef<IAssetTypeActions> Action)
+	/** Register settings objects. */
+	void RegisterSettings()
 	{
-		AssetTools.RegisterAssetTypeActions(Action);
-		RegisteredAssetTypeActions.Add(Action);
+		ISettingsModule* SettingsModule = FModuleManager::GetModulePtr<ISettingsModule>("Settings");
+
+		if (SettingsModule != nullptr)
+		{
+			// @todo sequencer: this should be moved into LevelSequenceEditor
+			SettingsModule->RegisterSettings("Project", "Plugins", "LevelSequencer",
+				LOCTEXT("LevelSequenceEditorSettingsName", "Level Sequencer"),
+				LOCTEXT("LevelSequenceEditorSettingsDescription", "Configure the Level Sequence Editor."),
+				GetMutableDefault<ULevelSequenceEditorSettings>()
+			);
+		}
 	}
 
+protected:
+
 	/** Unregisters asset tool actions. */
 	void UnregisterAssetTools()
 	{
@@ -131,6 +165,16 @@
 		}
 	}
 
+	/** Unregister details view customizations. */
+	void UnregisterCustomizations()
+	{
+		FPropertyEditorModule* PropertyModule = FModuleManager::GetModulePtr<FPropertyEditorModule>("PropertyEditor");
+		if (PropertyModule)
+		{
+			PropertyModule->UnregisterCustomPropertyTypeLayout(LevelSequencePlayingSettingsName);
+		}
+	}
+
 	/** Unregisters level editor extensions. */
 	void UnregisterLevelEditorExtensions()
 	{
@@ -141,56 +185,6 @@
 		}
 	}
 
-	/** Unregisters placement mode extensions. */
-	void UnregisterPlacementModeExtensions()
-	{
-		if (IPlacementModeModule::IsAvailable())
-		{
-			IPlacementModeModule::Get().UnregisterPlacementCategory("Cinematic");
-		}
-	}
-
-protected:
-
-	/** Register details view customizations. */
-	void RegisterCustomizations()
-	{
-		FPropertyEditorModule& PropertyModule = FModuleManager::LoadModuleChecked<FPropertyEditorModule>("PropertyEditor");
-		PropertyModule.RegisterCustomPropertyTypeLayout(FLevelSequencePlaybackSettings::StaticStruct()->GetFName(), FOnGetPropertyTypeCustomizationInstance::CreateStatic(&FLevelSequencePlaybackSettingsCustomization::MakeInstance));
-	}
-
-	/** Unregister details view customizations. */
-	void UnregisterCustomizations()
-	{
-		FPropertyEditorModule* PropertyModule = FModuleManager::GetModulePtr<FPropertyEditorModule>("PropertyEditor");
-		if (PropertyModule)
-		{
-			PropertyModule->UnregisterCustomPropertyTypeLayout(FLevelSequencePlaybackSettings::StaticStruct()->GetFName());
-		}
-	}
-
-protected:
-
-	/** Register menu extensions for the level editor toolbar. */
-	void RegisterMenuExtensions()
-	{
-		FLevelSequenceExtensionCommands::Register();
-
-		CommandList = MakeShareable(new FUICommandList);
-		CommandList->MapAction(FLevelSequenceExtensionCommands::Get().CreateNewLevelSequenceInLevel,
-			FExecuteAction::CreateStatic(&FLevelSequenceEditorModule::OnCreateActorInLevel)
-		);
-
-		// Create and register the level editor toolbar menu extension
-		CinematicsMenuExtender = MakeShareable(new FExtender);
-		CinematicsMenuExtender->AddMenuExtension("LevelEditorNewMatinee", EExtensionHook::First, CommandList, FMenuExtensionDelegate::CreateStatic([](FMenuBuilder& MenuBuilder){
-			MenuBuilder.AddMenuEntry(FLevelSequenceExtensionCommands::Get().CreateNewLevelSequenceInLevel);
-		}));
-
-		FLevelEditorModule& LevelEditorModule = FModuleManager::LoadModuleChecked<FLevelEditorModule>("LevelEditor");
-		LevelEditorModule.GetAllLevelEditorToolbarCinematicsMenuExtenders().Add(CinematicsMenuExtender);
-	}
-
 	/** Unregisters menu extensions for the level editor toolbar. */
 	void UnregisterMenuExtensions()
 	{
@@ -202,9 +196,32 @@
 		CinematicsMenuExtender = nullptr;
 		CommandList = nullptr;
 
-		FLevelSequenceExtensionCommands::Unregister();
+		FLevelSequenceEditorCommands::Unregister();
 	}
 
+	/** Unregisters placement mode extensions. */
+	void UnregisterPlacementModeExtensions()
+	{
+		if (IPlacementModeModule::IsAvailable())
+		{
+			IPlacementModeModule::Get().UnregisterPlacementCategory("Cinematic");
+		}
+	}
+
+	/** Unregister settings objects. */
+	void UnregisterSettings()
+	{
+		ISettingsModule* SettingsModule = FModuleManager::GetModulePtr<ISettingsModule>("Settings");
+
+		if (SettingsModule != nullptr)
+		{
+			// @todo sequencer: this should be moved into LevelSequenceEditor
+			SettingsModule->UnregisterSettings("Project", "Plugins", "LevelSequencer");
+		}
+	}
+
+protected:
+
 	/** Callback for creating a new level sequence asset in the level. */
 	static void OnCreateActorInLevel()
 	{
@@ -262,6 +279,9 @@
 	TSharedPtr<FExtender> CinematicsMenuExtender;
 
 	TSharedPtr<FUICommandList> CommandList;
+
+	/** Captured name of the FLevelSequencePlaybackSettings struct */
+	FName LevelSequencePlayingSettingsName;
 };
 
 
diff --git a/Engine/Plugins/MovieScene/LevelSequenceEditor/Source/LevelSequenceEditor/Private/LevelSequenceEditorPCH.h b/Engine/Plugins/MovieScene/LevelSequenceEditor/Source/LevelSequenceEditor/Private/LevelSequenceEditorPCH.h
--- a/Engine/Plugins/MovieScene/LevelSequenceEditor/Source/LevelSequenceEditor/Private/LevelSequenceEditorPCH.h
+++ b/Engine/Plugins/MovieScene/LevelSequenceEditor/Source/LevelSequenceEditor/Private/LevelSequenceEditorPCH.h
@@ -20,6 +20,9 @@
  *****************************************************************************/
 
 #include "LevelSequenceActions.h"
+#include "LevelSequenceEditorCommands.h"
+#include "LevelSequenceEditorSettings.h"
+#include "LevelSequenceEditorSpawnRegister.h"
 #include "LevelSequenceEditorToolkit.h"
 #include "LevelSequenceEditorStyle.h"
 #include "LevelSequencePlaybackSettingsCustomization.h"
diff --git a/Engine/Plugins/MovieScene/LevelSequenceEditor/Source/LevelSequenceEditor/Private/LevelSequenceEditorToolkit.cpp b/Engine/Plugins/MovieScene/LevelSequenceEditor/Source/LevelSequenceEditor/Private/LevelSequenceEditorToolkit.cpp
--- a/Engine/Plugins/MovieScene/LevelSequenceEditor/Source/LevelSequenceEditor/Private/LevelSequenceEditorToolkit.cpp
+++ b/Engine/Plugins/MovieScene/LevelSequenceEditor/Source/LevelSequenceEditor/Private/LevelSequenceEditorToolkit.cpp
@@ -6,18 +6,21 @@
 #include "Editor/LevelEditor/Public/ILevelViewport.h"
 #include "Editor/WorkspaceMenuStructure/Public/WorkspaceMenuStructure.h"
 #include "Editor/WorkspaceMenuStructure/Public/WorkspaceMenuStructureModule.h"
+#include "ClassIconFinder.h"
 #include "ISequencer.h"
 #include "ISequencerModule.h"
-#include "SDockTab.h"
-#include "MovieSceneBinding.h"
 #include "MovieScene.h"
-#include "MovieSceneSequenceInstance.h"
+#include "MovieSceneBinding.h"
 #include "MovieSceneMaterialTrack.h"
+#include "MovieSceneSequenceInstance.h"
 #include "ScopedTransaction.h"
-#include "ClassIconFinder.h"
 #include "SceneOutlinerModule.h"
 #include "SceneOutlinerPublicTypes.h"
-#include "LevelSequenceEditorSpawnRegister.h"
+#include "SDockTab.h"
+
+// @todo sequencer: hack: setting defaults for transform tracks
+#include "MovieScene3DTransformSection.h"
+#include "MovieScene3DTransformTrack.h"
 
 
 #define LOCTEXT_NAMESPACE "LevelSequenceEditor"
@@ -279,19 +282,161 @@ void FLevelSequenceEditorToolkit::AddActorsToSequencer(AActor*const* InActors, i
 	while (NumActors--)
 	{
 		AActor* ThisActor = *InActors;
+
 		if (!Sequencer->GetFocusedMovieSceneSequenceInstance()->FindObjectId(*ThisActor).IsValid())
 		{
-			Sequencer->CreateBinding(*ThisActor, ThisActor->GetActorLabel());
+			FGuid Binding = Sequencer->CreateBinding(*ThisActor, ThisActor->GetActorLabel());
+			Sequencer->UpdateRuntimeInstances();
+			AddDefaultTracksForActor(*ThisActor, Binding);
 		}
 
 		GEditor->SelectActor(ThisActor, true, true);
-
 		InActors++;
 	}
 
 	Sequencer->NotifyMovieSceneDataChanged();
 }
 
+void FLevelSequenceEditorToolkit::AddDefaultTracksForActor(AActor& Actor, const FGuid Binding)
+{
+	// get focused movie scene
+	UMovieSceneSequence* Sequence = Sequencer->GetFocusedMovieSceneSequence();
+
+	if (Sequence == nullptr)
+	{
+		return;
+	}
+
+	UMovieScene* MovieScene = Sequence->GetMovieScene();
+
+	if (MovieScene == nullptr)
+	{
+		return;
+	}
+
+	// add default tracks
+	for (const FLevelSequenceTrackSettings& TrackSettings : GetDefault<ULevelSequenceEditorSettings>()->TrackSettings)
+	{
+		UClass* MatchingActorClass = TrackSettings.MatchingActorClass.ResolveClass();
+
+		if ((MatchingActorClass == nullptr) || !Actor.IsA(MatchingActorClass))
+		{
+			continue;
+		}
+
+		// add tracks by type
+		for (const FStringClassReference& DefaultTrack : TrackSettings.DefaultTracks)
+		{
+			UClass* TrackClass = DefaultTrack.ResolveClass();
+
+			if (TrackClass != nullptr)
+			{
+				UMovieSceneTrack* NewTrack = MovieScene->AddTrack(TrackClass, Binding);
+				UMovieSceneSection* NewSection = NewTrack->CreateNewSection();
+
+				NewTrack->AddSection(*NewSection);
+
+				// @todo sequencer: hack: setting defaults for transform tracks
+				if (NewTrack->IsA(UMovieScene3DTransformTrack::StaticClass()))
+				{
+					auto TransformSection = Cast<UMovieScene3DTransformSection>(NewSection);
+
+					const FVector Location = Actor.GetActorLocation();
+					const FRotator Rotation = Actor.GetActorRotation();
+					const FVector Scale = Actor.GetActorScale();
+
+					TransformSection->SetDefault(FTransformKey(EKey3DTransformChannel::Translation, EAxis::X, Location.X, false /*bUnwindRotation*/));
+					TransformSection->SetDefault(FTransformKey(EKey3DTransformChannel::Translation, EAxis::Y, Location.Y, false /*bUnwindRotation*/));
+					TransformSection->SetDefault(FTransformKey(EKey3DTransformChannel::Translation, EAxis::Z, Location.Z, false /*bUnwindRotation*/));
+
+					TransformSection->SetDefault(FTransformKey(EKey3DTransformChannel::Rotation, EAxis::X, Rotation.Euler().X, false /*bUnwindRotation*/));
+					TransformSection->SetDefault(FTransformKey(EKey3DTransformChannel::Rotation, EAxis::Y, Rotation.Euler().Y, false /*bUnwindRotation*/));
+					TransformSection->SetDefault(FTransformKey(EKey3DTransformChannel::Rotation, EAxis::Z, Rotation.Euler().Z, false /*bUnwindRotation*/));
+
+					TransformSection->SetDefault(FTransformKey(EKey3DTransformChannel::Scale, EAxis::X, Scale.X, false /*bUnwindRotation*/));
+					TransformSection->SetDefault(FTransformKey(EKey3DTransformChannel::Scale, EAxis::Y, Scale.Y, false /*bUnwindRotation*/));
+					TransformSection->SetDefault(FTransformKey(EKey3DTransformChannel::Scale, EAxis::Z, Scale.Z, false /*bUnwindRotation*/));
+				}
+
+				Sequencer->UpdateRuntimeInstances();
+			}
+		}
+
+		// add tracks by property
+		for (const FLevelSequencePropertyTrackSettings& PropertyTrackSettings : TrackSettings.DefaultPropertyTracks)
+		{
+			TArray<UProperty*> PropertyPath;
+			UObject* PropertyOwner = &Actor;
+
+			// determine object hierarchy
+			TArray<FString> ComponentNames;
+			PropertyTrackSettings.ComponentPath.ParseIntoArray(ComponentNames, TEXT("."));
+
+			for (const FString& ComponentName : ComponentNames)
+			{
+				PropertyOwner = FindObjectFast<UObject>(PropertyOwner, *ComponentName);
+
+				if (PropertyOwner == nullptr)
+				{
+					return;
+				}
+			}
+
+			UStruct* PropertyOwnerClass = PropertyOwner->GetClass();
+
+			// determine property path
+			TArray<FString> PropertyNames;
+			PropertyTrackSettings.PropertyPath.ParseIntoArray(PropertyNames, TEXT("."));
+
+			for (const FString& PropertyName : PropertyNames)
+			{
+				UProperty* Property = PropertyOwnerClass->FindPropertyByName(*PropertyName);
+
+				if (Property != nullptr)
+				{
+					PropertyPath.Add(Property);
+				}
+
+				UStructProperty* StructProperty = Cast<UStructProperty>(Property);
+
+				if (StructProperty != nullptr)
+				{
+					PropertyOwnerClass = StructProperty->Struct;
+					continue;
+				}
+
+				UObjectProperty* ObjectProperty = Cast<UObjectProperty>(Property);
+
+				if (ObjectProperty != nullptr)
+				{
+					PropertyOwnerClass = ObjectProperty->PropertyClass;
+					continue;
+				}
+
+				break;
+			}
+
+			if (!Sequencer->CanKeyProperty(FCanKeyPropertyParams(Actor.GetClass(), PropertyPath)))
+			{
+				continue;
+			}
+
+			// key property
+			FKeyPropertyParams KeyPropertyParams(TArrayBuilder<UObject*>().Add(PropertyOwner), PropertyPath);
+			{
+				KeyPropertyParams.KeyParams.bCreateTrackIfMissing = true;
+				KeyPropertyParams.KeyParams.bCreateHandleIfMissing = true;
+				KeyPropertyParams.KeyParams.bCreateKeyIfUnchanged = true;
+				KeyPropertyParams.KeyParams.bCreateKeyIfEmpty = true;
+				KeyPropertyParams.KeyParams.bCreateKeyOnlyWhenAutoKeying = true;
+			}
+
+			Sequencer->KeyProperty(KeyPropertyParams);
+			Sequencer->UpdateRuntimeInstances();
+		}
+	}
+}
+
 void FLevelSequenceEditorToolkit::AddPosessActorMenuExtensions(FMenuBuilder& MenuBuilder)
 {
 	auto IsActorValidForPossession = [=](const AActor* Actor){
diff --git a/Engine/Plugins/MovieScene/LevelSequenceEditor/Source/LevelSequenceEditor/Private/LevelSequenceEditorToolkit.h b/Engine/Plugins/MovieScene/LevelSequenceEditor/Source/LevelSequenceEditor/Private/LevelSequenceEditorToolkit.h
--- a/Engine/Plugins/MovieScene/LevelSequenceEditor/Source/LevelSequenceEditor/Private/LevelSequenceEditorToolkit.h
+++ b/Engine/Plugins/MovieScene/LevelSequenceEditor/Source/LevelSequenceEditor/Private/LevelSequenceEditorToolkit.h
@@ -106,6 +106,9 @@ private:
 	/** Add the specified actors to the sequencer */
 	void AddActorsToSequencer(AActor*const* InActors, int32 NumActors);
 
+	/** Add default movie scene tracks for the given actor. */
+	void AddDefaultTracksForActor(AActor& Actor, const FGuid Binding);
+
 private:
 
 	/** Called when this toolkit is to be destroyed */
diff --git a/Engine/Plugins/MovieScene/LevelSequenceEditor/Source/LevelSequenceEditor/Private/Misc/LevelSequenceEditorSpawnRegister.cpp b/Engine/Plugins/MovieScene/LevelSequenceEditor/Source/LevelSequenceEditor/Private/Misc/LevelSequenceEditorSpawnRegister.cpp
--- a/Engine/Plugins/MovieScene/LevelSequenceEditor/Source/LevelSequenceEditor/Private/Misc/LevelSequenceEditorSpawnRegister.cpp
+++ b/Engine/Plugins/MovieScene/LevelSequenceEditor/Source/LevelSequenceEditor/Private/Misc/LevelSequenceEditorSpawnRegister.cpp
@@ -4,7 +4,6 @@
 #include "ISequencer.h"
 #include "ISequencerObjectChangeListener.h"
 #include "Editor/LevelEditor/Public/LevelEditor.h"
-#include "LevelSequenceEditorSpawnRegister.h"
 #include "MovieScene.h"
 #include "Kismet2/KismetEditorUtilities.h"
 #include "Kismet2/BlueprintEditorUtils.h"
@@ -23,7 +22,10 @@ FLevelSequenceEditorSpawnRegister::FLevelSequenceEditorSpawnRegister()
 	OnActorSelectionChangedHandle = LevelEditor.OnActorSelectionChanged().AddRaw(this, &FLevelSequenceEditorSpawnRegister::HandleActorSelectionChanged);
 
 	OnActorMovedHandle = GEditor->OnActorMoved().AddLambda([=](AActor* Actor){
-		HandleAnyPropertyChanged(*Actor);
+		if (SpawnedObjects.Contains(Actor))
+		{
+			HandleAnyPropertyChanged(*Actor);
+		}
 	});
 }
 
@@ -45,19 +47,24 @@ UObject* FLevelSequenceEditorSpawnRegister::SpawnObject(const FGuid& BindingId,
 {
 	TGuardValue<bool> Guard(bShouldClearSelectionCache, false);
 	UObject* NewObject = FLevelSequenceSpawnRegister::SpawnObject(BindingId, SequenceInstance, Player);
-
-	// Add an object listener for the spawned object to propagate changes back onto the spawnable default
-	TSharedPtr<ISequencer> Sequencer = WeakSequencer.Pin();
-
-	if (Sequencer.IsValid() && NewObject)
+	
+	if (NewObject)
 	{
-		Sequencer->GetObjectChangeListener().GetOnAnyPropertyChanged(*NewObject).AddSP(this, &FLevelSequenceEditorSpawnRegister::HandleAnyPropertyChanged);
+		SpawnedObjects.Add(NewObject);
 
-		// Select the actor if we think it should be selected
+		// Add an object listener for the spawned object to propagate changes back onto the spawnable default
+		TSharedPtr<ISequencer> Sequencer = WeakSequencer.Pin();
 		AActor* Actor = Cast<AActor>(NewObject);
-		if (Actor && SelectedSpawnedObjects.Contains(FMovieSceneSpawnRegisterKey(BindingId, SequenceInstance)))
+
+		if (Sequencer.IsValid() && Actor)
 		{
-			GEditor->SelectActor(Actor, true /*bSelected*/, true /*bNotify*/);
+			Sequencer->GetObjectChangeListener().GetOnAnyPropertyChanged(*NewObject).AddSP(this, &FLevelSequenceEditorSpawnRegister::HandleAnyPropertyChanged);
+
+			// Select the actor if we think it should be selected
+			if (Actor && SelectedSpawnedObjects.Contains(FMovieSceneSpawnRegisterKey(BindingId, SequenceInstance)))
+			{
+				GEditor->SelectActor(Actor, true /*bSelected*/, true /*bNotify*/);
+			}
 		}
 	}
 
@@ -76,6 +83,8 @@ void FLevelSequenceEditorSpawnRegister::PreDestroyObject(UObject& Object, const
 		SelectedSpawnedObjects.Add(FMovieSceneSpawnRegisterKey(BindingId, SequenceInstance));
 		GEditor->SelectActor(Actor, false /*bSelected*/, true /*bNotify*/);
 	}
+			
+	SpawnedObjects.Remove(&Object);
 
 	// Remove our object listener
 	TSharedPtr<ISequencer> Sequencer = WeakSequencer.Pin();
@@ -117,9 +126,6 @@ void FLevelSequenceEditorSpawnRegister::HandleActorSelectionChanged(const TArray
 
 void FLevelSequenceEditorSpawnRegister::HandleAnyPropertyChanged(UObject& SpawnedObject)
 {
-	//@todo: AndrewR to filter this to just spawned objects only
-	return;
-
 	using namespace EditorUtilities;
 
 	AActor* Actor = CastChecked<AActor>(&SpawnedObject);
diff --git a/Engine/Plugins/MovieScene/LevelSequenceEditor/Source/LevelSequenceEditor/Private/Misc/LevelSequenceEditorSpawnRegister.h b/Engine/Plugins/MovieScene/LevelSequenceEditor/Source/LevelSequenceEditor/Private/Misc/LevelSequenceEditorSpawnRegister.h
--- a/Engine/Plugins/MovieScene/LevelSequenceEditor/Source/LevelSequenceEditor/Private/Misc/LevelSequenceEditorSpawnRegister.h
+++ b/Engine/Plugins/MovieScene/LevelSequenceEditor/Source/LevelSequenceEditor/Private/Misc/LevelSequenceEditorSpawnRegister.h
@@ -3,7 +3,7 @@
 #pragma once
 
 #include "LevelSequenceSpawnRegister.h"
-
+#include "ObjectKey.h"
 
 class AActor;
 class FMovieSceneSequenceInstance;
@@ -65,6 +65,9 @@ private:
 	/** Set of spawn register keys for objects that should be selected if they are spawned. */
 	TSet<FMovieSceneSpawnRegisterKey> SelectedSpawnedObjects;
 
+	/** Set of currently spawned objects */
+	TSet<FObjectKey> SpawnedObjects;
+
 	/** True if we should clear the above selection cache when the editor selection has been changed. */
 	bool bShouldClearSelectionCache;
 
diff --git a/Engine/Source/Developer/LogVisualizer/Private/LogVisualizer.cpp b/Engine/Source/Developer/LogVisualizer/Private/LogVisualizer.cpp
--- a/Engine/Source/Developer/LogVisualizer/Private/LogVisualizer.cpp
+++ b/Engine/Source/Developer/LogVisualizer/Private/LogVisualizer.cpp
@@ -1,9 +1,7 @@
 // Copyright 1998-2016 Epic Games, Inc. All Rights Reserved.
 
 #include "LogVisualizer.h"
-#include "TimeSliderController.h"
 #include "VisualLoggerRenderingActor.h"
-#include "STimelinesContainer.h"
 #include "VisualLoggerCameraController.h"
 #if WITH_EDITOR
 #	include "Editor/UnrealEd/Public/EditorComponents.h"
diff --git a/Engine/Source/Developer/LogVisualizer/Private/LogVisualizer.h b/Engine/Source/Developer/LogVisualizer/Private/LogVisualizer.h
--- a/Engine/Source/Developer/LogVisualizer/Private/LogVisualizer.h
+++ b/Engine/Source/Developer/LogVisualizer/Private/LogVisualizer.h
@@ -111,5 +111,10 @@ public:
 #include "SVisualLoggerFilters.h"
 #include "SVisualLoggerView.h"
 #include "SVisualLoggerLogsList.h"
+#include "SVisualLoggerSectionOverlay.h"
 #include "SVisualLoggerStatusView.h"
-#include "STimeline.h"
+#include "SVisualLoggerTimeline.h"
+#include "SVisualLoggerTimelineBar.h"
+#include "SVisualLoggerTimelinesContainer.h"
+#include "SVisualLoggerTimeSlider.h"
+#include "VisualLoggerTimeSliderController.h"
\ No newline at end of file
diff --git a/Engine/Source/Developer/LogVisualizer/Private/SVisualLogger.cpp b/Engine/Source/Developer/LogVisualizer/Private/SVisualLogger.cpp
--- a/Engine/Source/Developer/LogVisualizer/Private/SVisualLogger.cpp
+++ b/Engine/Source/Developer/LogVisualizer/Private/SVisualLogger.cpp
@@ -8,7 +8,6 @@
 #include "DesktopPlatformModule.h"
 #include "MainFrame.h"
 #include "VisualLoggerCameraController.h"
-#include "TimeSliderController.h"
 #if WITH_EDITOR
 #	include "Editor/UnrealEd/Public/EditorComponents.h"
 #	include "Editor/UnrealEd/Public/EditorReimportHandler.h"
diff --git a/Engine/Source/Developer/LogVisualizer/Private/SVisualLoggerLogsList.cpp b/Engine/Source/Developer/LogVisualizer/Private/SVisualLoggerLogsList.cpp
--- a/Engine/Source/Developer/LogVisualizer/Private/SVisualLoggerLogsList.cpp
+++ b/Engine/Source/Developer/LogVisualizer/Private/SVisualLoggerLogsList.cpp
@@ -2,12 +2,8 @@
 
 #include "LogVisualizer.h"
 #include "SVisualLoggerView.h"
-#include "TimeSliderController.h"
 #include "ITimeSlider.h"
-#include "STimeSlider.h"
 #include "SSearchBox.h"
-#include "SSequencerSectionOverlay.h"
-#include "STimelinesContainer.h"
 #include "SVisualLoggerLogsList.h"
 #include "LogVisualizerSettings.h"
 
diff --git a/Engine/Source/Developer/LogVisualizer/Private/SVisualLoggerReport.cpp b/Engine/Source/Developer/LogVisualizer/Private/SVisualLoggerReport.cpp
--- a/Engine/Source/Developer/LogVisualizer/Private/SVisualLoggerReport.cpp
+++ b/Engine/Source/Developer/LogVisualizer/Private/SVisualLoggerReport.cpp
@@ -18,7 +18,7 @@
 #define LOCTEXT_NAMESPACE "SVisualLoggerReport"
 
 BEGIN_SLATE_FUNCTION_BUILD_OPTIMIZATION
-void SVisualLoggerReport::Construct(const FArguments& InArgs, TArray< TSharedPtr<class STimeline> >& InSelectedItems, TSharedPtr<class SVisualLoggerView> VisualLoggerView)
+void SVisualLoggerReport::Construct(const FArguments& InArgs, TArray< TSharedPtr<class SLogVisualizerTimeline> >& InSelectedItems, TSharedPtr<class SVisualLoggerView> VisualLoggerView)
 {
 	SelectedItems = InSelectedItems;
 	GenerateReportText();
@@ -91,7 +91,7 @@ void SVisualLoggerReport::GenerateReportText()
 
 
 	OutString.Append(TEXT("<RichText.HeaderText1>Report Details</>\n"));
-	for (TSharedPtr<class STimeline> LogItem : SelectedItems)
+	for (TSharedPtr<class SLogVisualizerTimeline> LogItem : SelectedItems)
 	{
 		TArray<FVisualLogEvent> AllEvents;
 
diff --git a/Engine/Source/Developer/LogVisualizer/Private/SVisualLoggerReport.h b/Engine/Source/Developer/LogVisualizer/Private/SVisualLoggerReport.h
--- a/Engine/Source/Developer/LogVisualizer/Private/SVisualLoggerReport.h
+++ b/Engine/Source/Developer/LogVisualizer/Private/SVisualLoggerReport.h
@@ -8,7 +8,7 @@ public:
 
 	SLATE_END_ARGS()
 
-	void Construct(const FArguments& InArgs, TArray< TSharedPtr<class STimeline> >&, TSharedPtr<class SVisualLoggerView> VisualLoggerView);
+	void Construct(const FArguments& InArgs, TArray< TSharedPtr<class SLogVisualizerTimeline> >&, TSharedPtr<class SVisualLoggerView> VisualLoggerView);
 
 	virtual ~SVisualLoggerReport();
 
@@ -16,7 +16,7 @@ protected:
 	void GenerateReportText();
 
 protected:
-	TArray< TSharedPtr<class STimeline> > SelectedItems;
+	TArray< TSharedPtr<class SLogVisualizerTimeline> > SelectedItems;
 	TArray<TSharedRef<ITextDecorator>> Decorators;
 	TSharedPtr<class SRichTextBlock> InteractiveRichText;
 	FText ReportText;
diff --git a/Engine/Source/Developer/LogVisualizer/Private/SVisualLoggerSectionOverlay.cpp b/Engine/Source/Developer/LogVisualizer/Private/SVisualLoggerSectionOverlay.cpp
new file mode 100644
--- /dev/null
+++ b/Engine/Source/Developer/LogVisualizer/Private/SVisualLoggerSectionOverlay.cpp
@@ -0,0 +1,39 @@
+// Copyright 1998-2016 Epic Games, Inc. All Rights Reserved.
+
+#include "LogVisualizer.h"
+
+
+int32 SVisualLoggerSectionOverlay::OnPaint( const FPaintArgs& Args, const FGeometry& AllottedGeometry, const FSlateRect& MyClippingRect, FSlateWindowElementList& OutDrawElements, int32 LayerId, const FWidgetStyle& InWidgetStyle, bool bParentEnabled ) const
+{
+	TimeSliderController->OnPaintSectionView( AllottedGeometry, MyClippingRect, OutDrawElements, LayerId, ShouldBeEnabled( bParentEnabled ), bDisplayTickLines.Get(), bDisplayScrubPosition.Get() );
+
+	return SCompoundWidget::OnPaint( Args, AllottedGeometry, MyClippingRect, OutDrawElements, LayerId, InWidgetStyle, bParentEnabled  );
+}
+
+
+FReply SVisualLoggerSectionOverlay::OnMouseButtonDown( const FGeometry& MyGeometry, const FPointerEvent& MouseEvent )
+{
+	return TimeSliderController->OnMouseButtonDown(*this, MyGeometry, MouseEvent);
+}
+
+
+FReply SVisualLoggerSectionOverlay::OnMouseButtonUp( const FGeometry& MyGeometry, const FPointerEvent& MouseEvent )
+{
+	return TimeSliderController->OnMouseButtonUp( *this,  MyGeometry, MouseEvent );
+}
+
+
+FReply SVisualLoggerSectionOverlay::OnMouseMove( const FGeometry& MyGeometry, const FPointerEvent& MouseEvent )
+{
+	return TimeSliderController->OnMouseMove( *this, MyGeometry, MouseEvent );
+}
+
+
+FReply SVisualLoggerSectionOverlay::OnMouseWheel( const FGeometry& MyGeometry, const FPointerEvent& MouseEvent )
+{
+	if (MouseEvent.IsLeftShiftDown() || MouseEvent.IsLeftControlDown())
+	{
+		return TimeSliderController->OnMouseWheel(*this, MyGeometry, MouseEvent);
+	}
+	return FReply::Unhandled();
+}
diff --git a/Engine/Source/Developer/LogVisualizer/Private/SVisualLoggerStatusView.cpp b/Engine/Source/Developer/LogVisualizer/Private/SVisualLoggerStatusView.cpp
--- a/Engine/Source/Developer/LogVisualizer/Private/SVisualLoggerStatusView.cpp
+++ b/Engine/Source/Developer/LogVisualizer/Private/SVisualLoggerStatusView.cpp
@@ -2,12 +2,8 @@
 
 #include "LogVisualizer.h"
 #include "SVisualLoggerView.h"
-#include "TimeSliderController.h"
 #include "ITimeSlider.h"
-#include "STimeSlider.h"
 #include "SSearchBox.h"
-#include "SSequencerSectionOverlay.h"
-#include "STimelinesContainer.h"
 #include "SVisualLoggerStatusView.h"
 
 #define LOCTEXT_NAMESPACE "SVisualLoggerStatusView"
diff --git a/Engine/Source/Developer/LogVisualizer/Private/SVisualLoggerTimeSlider.cpp b/Engine/Source/Developer/LogVisualizer/Private/SVisualLoggerTimeSlider.cpp
new file mode 100644
--- /dev/null
+++ b/Engine/Source/Developer/LogVisualizer/Private/SVisualLoggerTimeSlider.cpp
@@ -0,0 +1,46 @@
+// Copyright 1998-2016 Epic Games, Inc. All Rights Reserved.
+
+#include "LogVisualizer.h"
+
+#define LOCTEXT_NAMESPACE "STimeSlider"
+
+
+void SVisualLoggerTimeSlider::Construct( const SVisualLoggerTimeSlider::FArguments& InArgs, TSharedRef<ITimeSliderController> InTimeSliderController )
+{
+	TimeSliderController = InTimeSliderController;
+	bMirrorLabels = InArgs._MirrorLabels;
+}
+
+int32 SVisualLoggerTimeSlider::OnPaint( const FPaintArgs& Args, const FGeometry& AllottedGeometry, const FSlateRect& MyClippingRect, FSlateWindowElementList& OutDrawElements, int32 LayerId, const FWidgetStyle& InWidgetStyle, bool bParentEnabled ) const
+{
+	int32 NewLayer = TimeSliderController->OnPaintTimeSlider( bMirrorLabels, AllottedGeometry, MyClippingRect, OutDrawElements, LayerId, InWidgetStyle, bParentEnabled );
+
+	return FMath::Max( NewLayer, SCompoundWidget::OnPaint(Args, AllottedGeometry, MyClippingRect, OutDrawElements, NewLayer, InWidgetStyle, ShouldBeEnabled( bParentEnabled ) ) );
+}
+
+FReply SVisualLoggerTimeSlider::OnMouseButtonDown( const FGeometry& MyGeometry, const FPointerEvent& MouseEvent )
+{
+	return TimeSliderController->OnMouseButtonDown( *this, MyGeometry, MouseEvent );
+}
+
+FReply SVisualLoggerTimeSlider::OnMouseButtonUp( const FGeometry& MyGeometry, const FPointerEvent& MouseEvent )
+{
+	return TimeSliderController->OnMouseButtonUp( *this,  MyGeometry, MouseEvent );
+}
+
+FReply SVisualLoggerTimeSlider::OnMouseMove( const FGeometry& MyGeometry, const FPointerEvent& MouseEvent )
+{
+	return TimeSliderController->OnMouseMove( *this, MyGeometry, MouseEvent );
+}
+
+FVector2D SVisualLoggerTimeSlider::ComputeDesiredSize( float ) const
+{
+	return FVector2D(100, 22);
+}
+
+FReply SVisualLoggerTimeSlider::OnMouseWheel( const FGeometry& MyGeometry, const FPointerEvent& MouseEvent )
+{
+	return TimeSliderController->OnMouseWheel( *this, MyGeometry, MouseEvent );
+}
+
+#undef LOCTEXT_NAMESPACE
diff --git a/Engine/Source/Developer/LogVisualizer/Private/SVisualLoggerTimeSlider.h b/Engine/Source/Developer/LogVisualizer/Private/SVisualLoggerTimeSlider.h
new file mode 100644
--- /dev/null
+++ b/Engine/Source/Developer/LogVisualizer/Private/SVisualLoggerTimeSlider.h
@@ -0,0 +1,37 @@
+// Copyright 1998-2016 Epic Games, Inc. All Rights Reserved.
+
+#pragma once
+
+#include "ITimeSlider.h"
+class SVisualLoggerTimeSlider : public ITimeSlider
+{
+public:
+
+	SLATE_BEGIN_ARGS(SVisualLoggerTimeSlider)
+		: _MirrorLabels( false )
+	{}
+		/* If we should mirror the labels on the timeline */
+		SLATE_ARGUMENT( bool, MirrorLabels )
+	SLATE_END_ARGS()
+
+
+	/**
+	 * Construct the widget
+	 * 
+	 * @param InArgs   A declaration from which to construct the widget
+	 */
+	void Construct( const FArguments& InArgs, TSharedRef<ITimeSliderController> InTimeSliderController );
+
+protected:
+	// SWidget interface
+	virtual FVector2D ComputeDesiredSize(float) const override;
+	virtual int32 OnPaint( const FPaintArgs& Args, const FGeometry& AllottedGeometry, const FSlateRect& MyClippingRect, FSlateWindowElementList& OutDrawElements, int32 LayerId, const FWidgetStyle& InWidgetStyle, bool bParentEnabled ) const override;
+	virtual FReply OnMouseButtonDown( const FGeometry& MyGeometry, const FPointerEvent& MouseEvent ) override;
+	virtual FReply OnMouseButtonUp( const FGeometry& MyGeometry, const FPointerEvent& MouseEvent ) override;
+	virtual FReply OnMouseMove( const FGeometry& MyGeometry, const FPointerEvent& MouseEvent ) override;
+	virtual FReply OnMouseWheel( const FGeometry& MyGeometry, const FPointerEvent& MouseEvent ) override;
+private:
+	TSharedPtr<ITimeSliderController> TimeSliderController;
+	bool bMirrorLabels;
+};
+
diff --git a/Engine/Source/Developer/LogVisualizer/Private/SVisualLoggerTimelineBar.cpp b/Engine/Source/Developer/LogVisualizer/Private/SVisualLoggerTimelineBar.cpp
new file mode 100644
--- /dev/null
+++ b/Engine/Source/Developer/LogVisualizer/Private/SVisualLoggerTimelineBar.cpp
@@ -0,0 +1,268 @@
+// Copyright 1998-2016 Epic Games, Inc. All Rights Reserved.
+
+#include "LogVisualizer.h"
+#if WITH_EDITOR
+#	include "Editor/UnrealEd/Public/EditorComponents.h"
+#	include "Editor/UnrealEd/Public/EditorReimportHandler.h"
+#	include "Editor/UnrealEd/Public/TexAlignTools.h"
+#	include "Editor/UnrealEd/Public/TickableEditorObject.h"
+#	include "Editor/UnrealEd/Public/Editor.h"
+#	include "Editor/UnrealEd/Public/EditorViewportClient.h"
+#endif
+
+FReply SVisualLoggerTimelineBar::OnMouseButtonDown(const FGeometry& MyGeometry, const FPointerEvent& MouseEvent)
+{
+	TimelineOwner.Pin()->OnMouseButtonDown(MyGeometry, MouseEvent);
+	FReply Replay = TimeSliderController->OnMouseButtonDown(*this, MyGeometry, MouseEvent);
+	if (Replay.IsEventHandled())
+	{
+		FName RowName = TimelineOwner.Pin()->GetName();
+		FVisualLoggerDBRow& DBRow = FVisualLoggerDatabase::Get().GetRowByName(RowName);
+		const float ScrubPosition = TimeSliderController->GetTimeSliderArgs().ScrubPosition.Get();
+		const int32 ClosestItem = DBRow.GetClosestItem(ScrubPosition);
+		const auto& Items = DBRow.GetItems();
+		if (Items.IsValidIndex(ClosestItem))
+		{
+			TimeSliderController->CommitScrubPosition(Items[ClosestItem].Entry.TimeStamp, false);
+		}
+	}
+	return Replay;
+}
+
+FReply SVisualLoggerTimelineBar::OnMouseButtonUp(const FGeometry& MyGeometry, const FPointerEvent& MouseEvent)
+{
+	TimelineOwner.Pin()->OnMouseButtonUp(MyGeometry, MouseEvent);
+
+	FReply  Replay = TimeSliderController->OnMouseButtonUp(*this, MyGeometry, MouseEvent);
+	if (Replay.IsEventHandled())
+	{
+		FName RowName = TimelineOwner.Pin()->GetName();
+		FVisualLoggerDBRow& DBRow = FVisualLoggerDatabase::Get().GetRowByName(RowName);
+		const float ScrubPosition = TimeSliderController->GetTimeSliderArgs().ScrubPosition.Get();
+		const int32 ClosestItem = DBRow.GetClosestItem(ScrubPosition);
+		const auto& Items = DBRow.GetItems();
+		if (Items.IsValidIndex(ClosestItem))
+		{
+			TimeSliderController->CommitScrubPosition(Items[ClosestItem].Entry.TimeStamp, false);
+		}
+	}
+	return Replay;
+}
+
+FReply SVisualLoggerTimelineBar::OnMouseMove(const FGeometry& MyGeometry, const FPointerEvent& MouseEvent)
+{
+	return TimeSliderController->OnMouseMove(*this, MyGeometry, MouseEvent);
+}
+
+FReply SVisualLoggerTimelineBar::OnMouseButtonDoubleClick(const FGeometry& MyGeometry, const FPointerEvent& MouseEvent)
+{
+	FName RowName = TimelineOwner.Pin()->GetName();
+	FVisualLoggerDBRow& DBRow = FVisualLoggerDatabase::Get().GetRowByName(RowName);
+	UWorld* World = FLogVisualizer::Get().GetWorld();
+	if (World && MouseEvent.GetEffectingButton() == EKeys::LeftMouseButton)
+	{
+		FLogVisualizer::Get().UpdateCameraPosition(RowName, DBRow.GetCurrentItemIndex());
+		return FReply::Handled();
+	}
+
+	return FReply::Unhandled();
+}
+
+SVisualLoggerTimelineBar::~SVisualLoggerTimelineBar()
+{
+}
+
+void SVisualLoggerTimelineBar::Construct(const FArguments& InArgs, TSharedPtr<FVisualLoggerTimeSliderController> InTimeSliderController, TSharedPtr<SLogVisualizerTimeline> InTimelineOwner)
+{
+	TimeSliderController = InTimeSliderController;
+	TimelineOwner = InTimelineOwner;
+
+	TRange<float> LocalViewRange = TimeSliderController->GetTimeSliderArgs().ViewRange.Get();
+}
+
+FVector2D SVisualLoggerTimelineBar::ComputeDesiredSize(float) const
+{
+	return FVector2D(5000.0f, 20.0f);
+}
+
+int32 SVisualLoggerTimelineBar::OnPaint(const FPaintArgs& Args, const FGeometry& AllottedGeometry, const FSlateRect& MyClippingRect, FSlateWindowElementList& OutDrawElements, int32 LayerId, const FWidgetStyle& InWidgetStyle, bool bParentEnabled) const
+{
+	//@TODO: Optimize it like it was with old LogVisualizer, to draw everything much faster (SebaK)
+	int32 RetLayerId = LayerId;
+
+	FArrangedChildren ArrangedChildren(EVisibility::Visible);
+	ArrangeChildren(AllottedGeometry, ArrangedChildren);
+
+	TRange<float> LocalViewRange = TimeSliderController->GetTimeSliderArgs().ViewRange.Get();
+	float LocalScrubPosition = TimeSliderController->GetTimeSliderArgs().ScrubPosition.Get();
+
+	float ViewRange = LocalViewRange.Size<float>();
+	float PixelsPerInput = ViewRange > 0 ? AllottedGeometry.Size.X / ViewRange : 0;
+	float CurrentScrubLinePos = (LocalScrubPosition - LocalViewRange.GetLowerBoundValue()) * PixelsPerInput;
+	float BoxWidth = FMath::Max(0.08f * PixelsPerInput, 4.0f);
+
+	// Draw a region around the entire section area
+	FSlateDrawElement::MakeBox(
+		OutDrawElements,
+		RetLayerId++,
+		AllottedGeometry.ToPaintGeometry(),
+		FLogVisualizerStyle::Get().GetBrush("Sequencer.SectionArea.Background"),
+		MyClippingRect,
+		ESlateDrawEffect::None,
+		TimelineOwner.Pin()->IsSelected() ? FLinearColor(.2f, .2f, .2f, 0.5f) : FLinearColor(.1f, .1f, .1f, 0.5f)
+		);
+
+	const FSlateBrush* FillImage = FLogVisualizerStyle::Get().GetBrush("LogVisualizer.LogBar.EntryDefault");
+	static const FColor CurrentTimeColor(140, 255, 255, 255);
+	static const FColor ErrorTimeColor(255, 0, 0, 255);
+	static const FColor WarningTimeColor(255, 255, 0, 255);
+	static const FColor SelectedBarColor(255, 255, 255, 255);
+	const FSlateBrush* SelectedFillImage = FLogVisualizerStyle::Get().GetBrush("LogVisualizer.LogBar.Selected");
+
+	const ESlateDrawEffect::Type DrawEffects = ESlateDrawEffect::None;// bEnabled ? ESlateDrawEffect::None : ESlateDrawEffect::DisabledEffect;
+
+	TArray<float> ErrorTimes;
+	TArray<float> WarningTimes;
+	int32 EntryIndex = 0;
+
+	FVisualLoggerDBRow& DBRow = FVisualLoggerDatabase::Get().GetRowByName(TimelineOwner.Pin()->GetName());
+	auto &Entries = DBRow.GetItems();
+	while (EntryIndex < Entries.Num())
+	{
+		const FVisualLogEntry& Entry = Entries[EntryIndex].Entry;
+		if (Entry.TimeStamp < LocalViewRange.GetLowerBoundValue() || Entry.TimeStamp > LocalViewRange.GetUpperBoundValue())
+		{
+			EntryIndex++;
+			continue;
+		}
+
+		if (DBRow.IsItemVisible(EntryIndex)==false)
+		{
+			EntryIndex++;
+			continue;
+		}
+
+		// find bar width, connect all contiguous bars to draw them as one geometry (rendering optimization)
+		const float StartPos = (Entry.TimeStamp - LocalViewRange.GetLowerBoundValue()) * PixelsPerInput - 2;
+		float EndPos = (Entry.TimeStamp - LocalViewRange.GetLowerBoundValue()) * PixelsPerInput + 2;
+		int32 StartIndex = EntryIndex;
+		float LastEndX = MAX_FLT;
+		for (; StartIndex < Entries.Num(); ++StartIndex)
+		{
+			const FVisualLogEntry& CurrentEntry = Entries[StartIndex].Entry;
+			if (CurrentEntry.TimeStamp < LocalViewRange.GetLowerBoundValue() || CurrentEntry.TimeStamp > LocalViewRange.GetUpperBoundValue())
+			{
+				break;
+			}
+
+			if (DBRow.IsItemVisible(StartIndex) == false)
+			{
+				continue;
+			}
+
+			const TArray<FVisualLogLine>& LogLines = CurrentEntry.LogLines;
+			bool bAddedWarning = false;
+			bool bAddedError = false;
+			for (const FVisualLogLine& CurrentLine : LogLines)
+			{
+				if (CurrentLine.Verbosity <= ELogVerbosity::Error && !bAddedError)
+				{
+					ErrorTimes.AddUnique(CurrentEntry.TimeStamp);
+					bAddedError = true;
+				}
+				else if (CurrentLine.Verbosity == ELogVerbosity::Warning && !bAddedWarning)
+				{
+					WarningTimes.AddUnique(CurrentEntry.TimeStamp);
+					bAddedWarning = true;
+				}
+				if (bAddedError && bAddedWarning)
+				{
+					break;
+				}
+			}
+
+			const float CurrentStartPos = (CurrentEntry.TimeStamp - LocalViewRange.GetLowerBoundValue()) * PixelsPerInput - 2;
+			if (CurrentStartPos > EndPos)
+			{
+				break;
+			}
+			EndPos = (CurrentEntry.TimeStamp - LocalViewRange.GetLowerBoundValue()) * PixelsPerInput + 2;
+		}
+
+		if (EndPos - StartPos > 0)
+		{
+			const float BarWidth = (EndPos - StartPos);
+			FSlateDrawElement::MakeBox(
+				OutDrawElements,
+				RetLayerId,
+				AllottedGeometry.ToPaintGeometry(
+				FVector2D(StartPos, 0.0f),
+				FVector2D(BarWidth, AllottedGeometry.Size.Y)),
+				FillImage,
+				MyClippingRect,
+				DrawEffects,
+				CurrentTimeColor
+				);
+		}
+		EntryIndex = StartIndex;
+	}
+
+	if (WarningTimes.Num()) RetLayerId++;
+	for (auto CurrentTime : WarningTimes)
+	{
+		float LinePos = (CurrentTime - LocalViewRange.GetLowerBoundValue()) * PixelsPerInput;
+
+		FSlateDrawElement::MakeBox(
+			OutDrawElements,
+			RetLayerId,
+			AllottedGeometry.ToPaintGeometry(
+			FVector2D(LinePos - 3, 0.0f),
+			FVector2D(6, AllottedGeometry.Size.Y)),
+			FillImage,
+			MyClippingRect,
+			DrawEffects,
+			WarningTimeColor
+			);
+	}
+
+	if (ErrorTimes.Num()) RetLayerId++;
+	for (auto CurrentTime : ErrorTimes)
+	{
+		float LinePos = (CurrentTime - LocalViewRange.GetLowerBoundValue()) * PixelsPerInput;
+
+		FSlateDrawElement::MakeBox(
+			OutDrawElements,
+			RetLayerId,
+			AllottedGeometry.ToPaintGeometry(
+			FVector2D(LinePos - 3, 0.0f),
+			FVector2D(6, AllottedGeometry.Size.Y)),
+			FillImage,
+			MyClippingRect,
+			DrawEffects,
+			ErrorTimeColor
+			);
+	}
+
+	int32 BestItemIndex = DBRow.GetClosestItem(LocalScrubPosition);
+
+	if (TimelineOwner.Pin()->IsSelected() && DBRow.GetCurrentItemIndex() != INDEX_NONE)
+	{
+		const auto &HighlightedItemEntry = DBRow.GetCurrentItem();
+		float CurrentTime = HighlightedItemEntry.Entry.TimeStamp;
+		float LinePos = (CurrentTime - LocalViewRange.GetLowerBoundValue()) * PixelsPerInput;
+
+		FSlateDrawElement::MakeBox(
+			OutDrawElements,
+			++RetLayerId,
+			AllottedGeometry.ToPaintGeometry(
+			FVector2D(LinePos - 2, 0.0f),
+			FVector2D(4, AllottedGeometry.Size.Y)),
+			SelectedFillImage,
+			MyClippingRect,
+			ESlateDrawEffect::None,
+			SelectedBarColor
+			);
+	}
+
+	return RetLayerId;
+}
\ No newline at end of file
diff --git a/Engine/Source/Developer/LogVisualizer/Private/SVisualLoggerView.cpp b/Engine/Source/Developer/LogVisualizer/Private/SVisualLoggerView.cpp
--- a/Engine/Source/Developer/LogVisualizer/Private/SVisualLoggerView.cpp
+++ b/Engine/Source/Developer/LogVisualizer/Private/SVisualLoggerView.cpp
@@ -2,15 +2,13 @@
 
 #include "LogVisualizer.h"
 #include "SVisualLoggerView.h"
-#include "TimeSliderController.h"
 #include "ITimeSlider.h"
-#include "STimeSlider.h"
 #include "SSearchBox.h"
-#include "SSequencerSectionOverlay.h"
-#include "STimelinesContainer.h"
+
 
 #define LOCTEXT_NAMESPACE "SVisualLoggerFilters"
 
+
 class SInputCatcherOverlay : public SOverlay
 {
 public:
@@ -77,8 +75,8 @@ void SVisualLoggerView::Construct(const FArguments& InArgs, const TSharedRef<FUI
 
 	// Create the top and bottom sliders
 	const bool bMirrorLabels = true;
-	TSharedRef<ITimeSlider> TopTimeSlider = SNew(STimeSlider, FLogVisualizer::Get().GetTimeSliderController().ToSharedRef()).MirrorLabels(bMirrorLabels);
-	TSharedRef<ITimeSlider> BottomTimeSlider = SNew(STimeSlider, FLogVisualizer::Get().GetTimeSliderController().ToSharedRef()).MirrorLabels(bMirrorLabels);
+	TSharedRef<ITimeSlider> TopTimeSlider = SNew(SVisualLoggerTimeSlider, FLogVisualizer::Get().GetTimeSliderController().ToSharedRef()).MirrorLabels(bMirrorLabels);
+	TSharedRef<ITimeSlider> BottomTimeSlider = SNew(SVisualLoggerTimeSlider, FLogVisualizer::Get().GetTimeSliderController().ToSharedRef()).MirrorLabels(bMirrorLabels);
 
 	TSharedRef<SScrollBar> ScrollBar =
 		SNew(SScrollBar)
@@ -194,7 +192,7 @@ void SVisualLoggerView::Construct(const FArguments& InArgs, const TSharedRef<FUI
 							.ExternalScrollbar(ScrollBar)
 							+ SScrollBox::Slot()
 							[
-								SAssignNew(TimelinesContainer, STimelinesContainer, SharedThis(this), FLogVisualizer::Get().GetTimeSliderController().ToSharedRef())
+								SAssignNew(TimelinesContainer, SVisualLoggerTimelinesContainer, SharedThis(this), FLogVisualizer::Get().GetTimeSliderController().ToSharedRef())
 							]
 						]
 						+ SOverlay::Slot()
@@ -343,13 +341,13 @@ FCursorReply SVisualLoggerView::OnCursorQuery(const FGeometry& MyGeometry, const
 
 TSharedRef<SWidget> SVisualLoggerView::MakeClassesFilterMenu()
 {
-	const TArray< TSharedPtr<class STimeline> >& AllTimelines = TimelinesContainer->GetAllNodes();
+	const TArray< TSharedPtr<class SLogVisualizerTimeline> >& AllTimelines = TimelinesContainer->GetAllNodes();
 
 	FMenuBuilder MenuBuilder(true, NULL);
 
 	TArray<FString> UniqueClasses;
 	MenuBuilder.BeginSection(TEXT("Graphs"));
-	for (TSharedPtr<class STimeline> CurrentTimeline : AllTimelines)
+	for (TSharedPtr<class SLogVisualizerTimeline> CurrentTimeline : AllTimelines)
 	{
 		FString OwnerClassName = CurrentTimeline->GetOwnerClassName().ToString();
 		if (UniqueClasses.Find(OwnerClassName) == INDEX_NONE)
diff --git a/Engine/Source/Developer/LogVisualizer/Private/SVisualLoggerView.h b/Engine/Source/Developer/LogVisualizer/Private/SVisualLoggerView.h
--- a/Engine/Source/Developer/LogVisualizer/Private/SVisualLoggerView.h
+++ b/Engine/Source/Developer/LogVisualizer/Private/SVisualLoggerView.h
@@ -40,7 +40,7 @@ public:
 	void ResetData();
 
 protected:
-	TSharedPtr<class STimelinesContainer> TimelinesContainer;
+	TSharedPtr<class SVisualLoggerTimelinesContainer> TimelinesContainer;
 	TSharedPtr<class SSplitter> SearchSplitter;
 	TSharedPtr<class SScrollBox> ScrollBox;
 	TSharedPtr<class SSearchBox> SearchBox;
diff --git a/Engine/Source/Developer/LogVisualizer/Private/VisualLoggerCanvasRenderer.cpp b/Engine/Source/Developer/LogVisualizer/Private/VisualLoggerCanvasRenderer.cpp
--- a/Engine/Source/Developer/LogVisualizer/Private/VisualLoggerCanvasRenderer.cpp
+++ b/Engine/Source/Developer/LogVisualizer/Private/VisualLoggerCanvasRenderer.cpp
@@ -6,7 +6,7 @@
 #include "Editor/UnrealEd/Classes/Editor/EditorEngine.h"
 #include "VisualLoggerCanvasRenderer.h"
 #include "Debug/ReporterGraph.h"
-#include "TimeSliderController.h"
+
 
 namespace LogVisualizer
 {
diff --git a/Engine/Source/Developer/LogVisualizer/Private/VisualLoggerRenderingActor.cpp b/Engine/Source/Developer/LogVisualizer/Private/VisualLoggerRenderingActor.cpp
--- a/Engine/Source/Developer/LogVisualizer/Private/VisualLoggerRenderingActor.cpp
+++ b/Engine/Source/Developer/LogVisualizer/Private/VisualLoggerRenderingActor.cpp
@@ -12,7 +12,7 @@
 #include "VisualLoggerRenderingActor.h"
 #include "VisualLoggerRenderingComponent.h"
 #include "DebugRenderSceneProxy.h"
-#include "STimeline.h"
+#include "SVisualLoggerTimeline.h"
 
 class UVisualLoggerRenderingComponent;
 class FVisualLoggerSceneProxy : public FDebugRenderSceneProxy
diff --git a/Engine/Source/Developer/LogVisualizer/Private/VisualLoggerTimeSliderController.cpp b/Engine/Source/Developer/LogVisualizer/Private/VisualLoggerTimeSliderController.cpp
new file mode 100644
--- /dev/null
+++ b/Engine/Source/Developer/LogVisualizer/Private/VisualLoggerTimeSliderController.cpp
@@ -0,0 +1,697 @@
+// Copyright 1998-2016 Epic Games, Inc. All Rights Reserved.
+
+#include "LogVisualizer.h"
+
+
+#define LOCTEXT_NAMESPACE "TimeSlider"
+
+
+namespace ScrubConstants
+{
+	/** The minimum amount of pixels between each major ticks on the widget */
+	const int32 MinPixelsPerDisplayTick = 5;
+
+	/**The smallest number of units between between major tick marks */
+	const float MinDisplayTickSpacing = 0.001f;
+}
+
+
+/** Utility struct for converting between scrub range space and local/absolute screen space */
+struct FScrubRangeToScreen
+{
+	FVector2D WidgetSize;
+
+	TRange<float> ViewInput;
+	float ViewInputRange;
+	float PixelsPerInput;
+
+	FScrubRangeToScreen(TRange<float> InViewInput, const FVector2D& InWidgetSize )
+	{
+		WidgetSize = InWidgetSize;
+
+		ViewInput = InViewInput;
+		ViewInputRange = ViewInput.Size<float>();
+		PixelsPerInput = ViewInputRange > 0 ? ( WidgetSize.X / ViewInputRange ) : 0;
+	}
+
+	/** Local Widget Space -> Curve Input domain. */
+	float LocalXToInput(float ScreenX) const
+	{
+		float LocalX = ScreenX;
+		return (LocalX/PixelsPerInput) + ViewInput.GetLowerBoundValue();
+	}
+
+	/** Curve Input domain -> local Widget Space */
+	float InputToLocalX(float Input) const
+	{
+		return (Input - ViewInput.GetLowerBoundValue()) * PixelsPerInput;
+	}
+};
+
+
+/**
+ * Gets the the next spacing value in the series 
+ * to determine a good spacing value
+ * E.g, .001,.005,.010,.050,.100,.500,1.000,etc
+ */
+static float GetNextSpacing( uint32 CurrentStep )
+{
+	if(CurrentStep & 0x01) 
+	{
+		// Odd numbers
+		return FMath::Pow( 10.f, 0.5f*((float)(CurrentStep-1)) + 1.f );
+	}
+	else 
+	{
+		// Even numbers
+		return 0.5f * FMath::Pow( 10.f, 0.5f*((float)(CurrentStep)) + 1.f );
+	}
+}
+
+/**
+ * Determines the optimal spacing between tick marks in the slider for a given pixel density
+ * Increments until a minimum amount of slate units specified by MinTick is reached
+ * 
+ * @param InPixelsPerInput	The density of pixels between each input
+ * @param MinTick			The minimum slate units per tick allowed
+ * @param MinTickSpacing	The minimum tick spacing in time units allowed
+ * @return the optimal spacing in time units
+ */
+float DetermineOptimalSpacing( float InPixelsPerInput, uint32 MinTick, float MinTickSpacing )
+{
+	uint32 CurStep = 0;
+
+	// Start with the smallest spacing
+	float Spacing = MinTickSpacing;
+
+	while( Spacing * InPixelsPerInput < MinTick )
+	{
+		Spacing = MinTickSpacing * GetNextSpacing( CurStep );
+		CurStep++;
+	}
+
+	return Spacing;
+}
+
+
+FVisualLoggerTimeSliderController::FVisualLoggerTimeSliderController(const FVisualLoggerTimeSliderArgs& InArgs)
+	: TimeSliderArgs( InArgs )
+	, DistanceDragged( 0.0f )
+	, bDraggingScrubber( false )
+	, bPanning( false )
+{
+	ScrubHandleUp = FEditorStyle::GetBrush( TEXT( "Sequencer.Timeline.ScrubHandleUp" ) ); 
+	ScrubHandleDown = FEditorStyle::GetBrush( TEXT( "Sequencer.Timeline.ScrubHandleDown" ) );
+	CursorBackground = FEditorStyle::GetBrush("Sequencer.SectionArea.Background");
+}
+
+void FVisualLoggerTimeSliderController::SetTimesliderArgs(const FVisualLoggerTimeSliderArgs& InArgs)
+{
+	TimeSliderArgs = InArgs;
+}
+
+struct FDrawTickArgs
+{
+	/** Geometry of the area */
+	FGeometry AllottedGeometry;
+	/** Clipping rect of the area */
+	FSlateRect ClippingRect;
+	/** Color of each tick */
+	FLinearColor TickColor;
+	/** Offset in Y where to start the tick */
+	float TickOffset;
+	/** Height in of major ticks */
+	float MajorTickHeight;
+	/** Start layer for elements */
+	int32 StartLayer;
+	/** Draw effects to apply */
+	ESlateDrawEffect::Type DrawEffects;
+	/** Whether or not to only draw major ticks */
+	bool bOnlyDrawMajorTicks;
+	/** Whether or not to mirror labels */
+	bool bMirrorLabels;
+	
+};
+
+void FVisualLoggerTimeSliderController::DrawTicks( FSlateWindowElementList& OutDrawElements, const struct FScrubRangeToScreen& RangeToScreen, FDrawTickArgs& InArgs ) const
+{
+	const float Spacing = DetermineOptimalSpacing( RangeToScreen.PixelsPerInput, ScrubConstants::MinPixelsPerDisplayTick, ScrubConstants::MinDisplayTickSpacing );
+
+	// Sub divisions
+	// @todo Sequencer may need more robust calculation
+	const int32 Divider = 10;
+	// For slightly larger halfway tick mark
+	const int32 HalfDivider = Divider / 2;
+	// Find out where to start from
+	int32 OffsetNum = FMath::FloorToInt(RangeToScreen.ViewInput.GetLowerBoundValue() / Spacing);
+	
+	FSlateFontInfo SmallLayoutFont( FPaths::EngineContentDir() / TEXT("Slate/Fonts/Roboto-Regular.ttf"), 8 );
+
+	TArray<FVector2D> LinePoints;
+	LinePoints.AddUninitialized(2);
+
+	float Seconds = 0;
+	while( (Seconds = OffsetNum*Spacing) < RangeToScreen.ViewInput.GetUpperBoundValue() )
+	{
+		// X position local to start of the widget area
+		float XPos = RangeToScreen.InputToLocalX( Seconds );
+		uint32 AbsOffsetNum = FMath::Abs(OffsetNum);
+
+		if ( AbsOffsetNum % Divider == 0 )
+		{
+			FVector2D Offset( XPos, InArgs.TickOffset );
+			FVector2D TickSize( 1.0f, InArgs.MajorTickHeight );
+
+			LinePoints[0] = FVector2D(1.0f,1.0f);
+			LinePoints[1] = TickSize;
+
+			// lines should not need anti-aliasing
+			const bool bAntiAliasLines = false;
+
+			// Draw each tick mark
+			FSlateDrawElement::MakeLines(
+				OutDrawElements,
+				InArgs.StartLayer,
+				InArgs.AllottedGeometry.ToPaintGeometry( Offset, TickSize ),
+				LinePoints,
+				InArgs.ClippingRect,
+				InArgs.DrawEffects,
+				InArgs.TickColor,
+				false
+				);
+
+			if( !InArgs.bOnlyDrawMajorTicks )
+			{
+				FString FrameString = Spacing == ScrubConstants::MinDisplayTickSpacing ? FString::Printf( TEXT("%.3f"), Seconds ) : FString::Printf( TEXT("%.2f"), Seconds );
+
+				// Space the text between the tick mark but slightly above
+				const TSharedRef< FSlateFontMeasure > FontMeasureService = FSlateApplication::Get().GetRenderer()->GetFontMeasureService();
+				FVector2D TextSize = FontMeasureService->Measure(FrameString, SmallLayoutFont);
+				FVector2D TextOffset( XPos-(TextSize.X*0.5f), InArgs.bMirrorLabels ? TextSize.Y :  FMath::Abs( InArgs.AllottedGeometry.Size.Y - (InArgs.MajorTickHeight+TextSize.Y) ) );
+
+				FSlateDrawElement::MakeText(
+					OutDrawElements,
+					InArgs.StartLayer+1, 
+					InArgs.AllottedGeometry.ToPaintGeometry( TextOffset, TextSize ), 
+					FrameString, 
+					SmallLayoutFont, 
+					InArgs.ClippingRect, 
+					InArgs.DrawEffects,
+					InArgs.TickColor 
+				);
+			}
+		}
+		else if( !InArgs.bOnlyDrawMajorTicks )
+		{
+			// Compute the size of each tick mark.  If we are half way between to visible values display a slightly larger tick mark
+			const float MinorTickHeight = AbsOffsetNum % HalfDivider == 0 ? 7.0f : 4.0f;
+
+			FVector2D Offset(XPos, InArgs.bMirrorLabels ? 0.0f : FMath::Abs( InArgs.AllottedGeometry.Size.Y - MinorTickHeight ) );
+			FVector2D TickSize(1, MinorTickHeight);
+
+			LinePoints[0] = FVector2D(1.0f,1.0f);
+			LinePoints[1] = TickSize;
+
+			const bool bAntiAlias = false;
+			// Draw each sub mark
+			FSlateDrawElement::MakeLines(
+				OutDrawElements,
+				InArgs.StartLayer,
+				InArgs.AllottedGeometry.ToPaintGeometry( Offset, TickSize ),
+				LinePoints,
+				InArgs.ClippingRect,
+				InArgs.DrawEffects,
+				InArgs.TickColor,
+				bAntiAlias
+			);
+		}
+		// Advance to next tick mark
+		++OffsetNum;
+	}
+}
+
+
+int32 FVisualLoggerTimeSliderController::OnPaintTimeSlider( bool bMirrorLabels, const FGeometry& AllottedGeometry, const FSlateRect& MyClippingRect, FSlateWindowElementList& OutDrawElements, int32 LayerId, const FWidgetStyle& InWidgetStyle, bool bParentEnabled ) const
+{
+	const bool bEnabled = bParentEnabled;
+	const ESlateDrawEffect::Type DrawEffects = bEnabled ? ESlateDrawEffect::None : ESlateDrawEffect::DisabledEffect;
+
+	TRange<float> LocalViewRange = TimeSliderArgs.ViewRange.Get();
+	const float LocalViewRangeMin = LocalViewRange.GetLowerBoundValue();
+	const float LocalViewRangeMax = LocalViewRange.GetUpperBoundValue();
+	const float LocalSequenceLength = LocalViewRangeMax-LocalViewRangeMin;
+	
+	FVector2D Scale = FVector2D(1.0f,1.0f);
+	if ( LocalSequenceLength > 0)
+	{
+		FScrubRangeToScreen RangeToScreen( LocalViewRange, AllottedGeometry.Size );
+	
+		const float MajorTickHeight = 9.0f;
+	
+		FDrawTickArgs Args;
+		Args.AllottedGeometry = AllottedGeometry;
+		Args.bMirrorLabels = bMirrorLabels;
+		Args.bOnlyDrawMajorTicks = false;
+		Args.TickColor = FLinearColor::White;
+		Args.ClippingRect = MyClippingRect;
+		Args.DrawEffects = DrawEffects;
+		Args.StartLayer = LayerId;
+		Args.TickOffset = bMirrorLabels ? 0.0f : FMath::Abs( AllottedGeometry.Size.Y - MajorTickHeight );
+		Args.MajorTickHeight = MajorTickHeight;
+
+		DrawTicks( OutDrawElements, RangeToScreen, Args );
+
+		const float HandleSize = 13.0f;
+		float HalfSize = FMath::TruncToFloat(HandleSize/2.0f);
+
+		// Draw the scrub handle
+		const float XPos = RangeToScreen.InputToLocalX( TimeSliderArgs.ScrubPosition.Get() );
+
+		// Draw cursor size
+		const float CursorHalfSize = TimeSliderArgs.CursorSize.Get() * 0.5f;
+		const int32 CursorLayer = LayerId + 2;
+		const float CursorHalfLength = AllottedGeometry.Size.X * CursorHalfSize;
+		FPaintGeometry CursorGeometry = AllottedGeometry.ToPaintGeometry(FVector2D(XPos - CursorHalfLength, 0), FVector2D(2 * CursorHalfLength, AllottedGeometry.Size.Y));
+
+		FLinearColor CursorColor = InWidgetStyle.GetColorAndOpacityTint();
+		CursorColor.A = CursorColor.A*0.08f;
+		CursorColor.B *= 0.1f;
+		CursorColor.G *= 0.2f;
+		FSlateDrawElement::MakeBox(
+			OutDrawElements,
+			CursorLayer,
+			CursorGeometry,
+			CursorBackground,
+			MyClippingRect,
+			DrawEffects,
+			CursorColor
+			);
+
+		// Should draw above the text
+		const int32 ArrowLayer = LayerId + 3;
+		FPaintGeometry MyGeometry =	AllottedGeometry.ToPaintGeometry( FVector2D( XPos-HalfSize, 0 ), FVector2D( HandleSize, AllottedGeometry.Size.Y ) );
+		FLinearColor ScrubColor = InWidgetStyle.GetColorAndOpacityTint();
+
+		// @todo Sequencer this color should be specified in the style
+		ScrubColor.A = ScrubColor.A*0.5f;
+		ScrubColor.B *= 0.1f;
+		ScrubColor.G *= 0.2f;
+		FSlateDrawElement::MakeBox( 
+			OutDrawElements,
+			ArrowLayer, 
+			MyGeometry,
+			bMirrorLabels ? ScrubHandleUp : ScrubHandleDown,
+			MyClippingRect, 
+			DrawEffects, 
+			ScrubColor
+			);
+
+		return ArrowLayer;
+	}
+
+	return LayerId;
+}
+
+FReply FVisualLoggerTimeSliderController::OnMouseButtonDown( SWidget& WidgetOwner, const FGeometry& MyGeometry, const FPointerEvent& MouseEvent )
+{
+	bool bHandleLeftMouseButton = MouseEvent.GetEffectingButton() == EKeys::LeftMouseButton;
+	bool bHandleRightMouseButton = MouseEvent.GetEffectingButton() == EKeys::RightMouseButton && TimeSliderArgs.AllowZoom;
+	
+	DistanceDragged = 0;
+
+	if ( bHandleLeftMouseButton )
+	{
+		// Always capture mouse if we left or right click on the widget
+		FScrubRangeToScreen RangeToScreen(TimeSliderArgs.ViewRange.Get(), MyGeometry.Size);
+		FVector2D CursorPos = MyGeometry.AbsoluteToLocal(MouseEvent.GetLastScreenSpacePosition());
+		float NewValue = RangeToScreen.LocalXToInput(CursorPos.X);
+
+		CommitScrubPosition(NewValue, /*bIsScrubbing=*/false);
+		return FReply::Handled().CaptureMouse( WidgetOwner.AsShared() ).PreventThrottling();
+	}
+	else if ( bHandleRightMouseButton )
+	{
+		return FReply::Handled().CaptureMouse(WidgetOwner.AsShared());
+	}
+
+	return FReply::Unhandled();
+}
+
+FReply FVisualLoggerTimeSliderController::OnMouseButtonUp( SWidget& WidgetOwner, const FGeometry& MyGeometry, const FPointerEvent& MouseEvent )
+{
+	bool bHandleLeftMouseButton = MouseEvent.GetEffectingButton() == EKeys::LeftMouseButton && WidgetOwner.HasMouseCapture();
+	bool bHandleRightMouseButton = MouseEvent.GetEffectingButton() == EKeys::RightMouseButton && WidgetOwner.HasMouseCapture() && TimeSliderArgs.AllowZoom ;
+	
+	if ( bHandleRightMouseButton )
+	{
+		if (!bPanning)
+		{
+			// return unhandled in case our parent wants to use our right mouse button to open a context menu
+			return FReply::Unhandled().ReleaseMouseCapture();
+		}
+		
+		bPanning = false;
+		FReply::Handled().CaptureMouse(WidgetOwner.AsShared()).UseHighPrecisionMouseMovement(WidgetOwner.AsShared());
+
+		return FReply::Handled().ReleaseMouseCapture();
+	}
+	else if ( bHandleLeftMouseButton )
+	{
+		if( bDraggingScrubber )
+		{
+			TimeSliderArgs.OnEndScrubberMovement.ExecuteIfBound();
+		}
+		else
+		{
+			FScrubRangeToScreen RangeToScreen( TimeSliderArgs.ViewRange.Get(), MyGeometry.Size );
+			FVector2D CursorPos = MyGeometry.AbsoluteToLocal(MouseEvent.GetLastScreenSpacePosition());
+			float NewValue = RangeToScreen.LocalXToInput(CursorPos.X);
+
+			CommitScrubPosition( NewValue, /*bIsScrubbing=*/false );
+		}
+
+		bDraggingScrubber = false;
+		return FReply::Handled().ReleaseMouseCapture();
+
+	}
+
+	return FReply::Unhandled();
+}
+
+FReply FVisualLoggerTimeSliderController::OnMouseMove( SWidget& WidgetOwner, const FGeometry& MyGeometry, const FPointerEvent& MouseEvent )
+{
+	if ( WidgetOwner.HasMouseCapture() )
+	{
+		if (MouseEvent.IsMouseButtonDown(EKeys::RightMouseButton))
+		{
+			if (!bPanning)
+			{
+				DistanceDragged += FMath::Abs( MouseEvent.GetCursorDelta().X );
+				if ( DistanceDragged > FSlateApplication::Get().GetDragTriggerDistance() )
+				{
+					FReply::Handled().CaptureMouse(WidgetOwner.AsShared()).UseHighPrecisionMouseMovement(WidgetOwner.AsShared());
+					SoftwareCursorPosition = MyGeometry.AbsoluteToLocal(MouseEvent.GetLastScreenSpacePosition());
+					bPanning = true;
+				}
+			}
+			else
+			{
+				SoftwareCursorPosition = MyGeometry.AbsoluteToLocal(MouseEvent.GetLastScreenSpacePosition());
+
+				TRange<float> LocalViewRange = TimeSliderArgs.ViewRange.Get();
+				float LocalViewRangeMin = LocalViewRange.GetLowerBoundValue();
+				float LocalViewRangeMax = LocalViewRange.GetUpperBoundValue();
+
+				FScrubRangeToScreen ScaleInfo( LocalViewRange, MyGeometry.Size );
+				FVector2D ScreenDelta = MouseEvent.GetCursorDelta();
+				FVector2D InputDelta;
+				InputDelta.X = ScreenDelta.X/ScaleInfo.PixelsPerInput;
+
+				float NewViewOutputMin = LocalViewRangeMin - InputDelta.X;
+				float NewViewOutputMax = LocalViewRangeMax - InputDelta.X;
+
+				float LocalClampMin = TimeSliderArgs.ClampRange.Get().GetLowerBoundValue();
+				float LocalClampMax = TimeSliderArgs.ClampRange.Get().GetUpperBoundValue();
+
+				// Clamp the range
+				if ( NewViewOutputMin < LocalClampMin )
+				{
+					NewViewOutputMin = LocalClampMin;
+				}
+			
+				if ( NewViewOutputMax > LocalClampMax )
+				{
+					NewViewOutputMax = LocalClampMax;
+				}
+
+				TimeSliderArgs.OnViewRangeChanged.ExecuteIfBound(TRange<float>(NewViewOutputMin, NewViewOutputMax), EViewRangeInterpolation::Immediate);
+				if (Scrollbar.IsValid())
+				{
+					float InOffsetFraction = (NewViewOutputMin - LocalClampMin) / (LocalClampMax - LocalClampMin);
+					float InThumbSizeFraction = (NewViewOutputMax - NewViewOutputMin) / (LocalClampMax - LocalClampMin);
+					Scrollbar->SetState(InOffsetFraction, InThumbSizeFraction);
+				}
+
+				if( !TimeSliderArgs.ViewRange.IsBound() )
+				{	
+					// The  output is not bound to a delegate so we'll manage the value ourselves
+					TimeSliderArgs.ViewRange.Set( TRange<float>( NewViewOutputMin, NewViewOutputMax ) );
+				}
+			}
+		}
+		else if (MouseEvent.IsMouseButtonDown( EKeys::LeftMouseButton ))
+		{
+			if ( !bDraggingScrubber )
+			{
+				DistanceDragged += FMath::Abs( MouseEvent.GetCursorDelta().X );
+				if ( DistanceDragged > 0/*FSlateApplication::Get().GetDragTriggerDistance()*/ )
+				{
+					bDraggingScrubber = true;
+					TimeSliderArgs.OnBeginScrubberMovement.ExecuteIfBound();
+				}
+			}
+			else
+			{
+				FScrubRangeToScreen RangeToScreen( TimeSliderArgs.ViewRange.Get(), MyGeometry.Size );
+				FVector2D CursorPos = MyGeometry.AbsoluteToLocal( MouseEvent.GetLastScreenSpacePosition() );
+				float NewValue = RangeToScreen.LocalXToInput( CursorPos.X );
+
+				float LocalClampMin = TimeSliderArgs.ClampRange.Get().GetLowerBoundValue();
+				float LocalClampMax = TimeSliderArgs.ClampRange.Get().GetUpperBoundValue();
+
+				if (NewValue < LocalClampMin)
+				{
+					NewValue = LocalClampMin;
+				}
+
+				if (NewValue > LocalClampMax)
+				{
+					NewValue = LocalClampMax;
+				}
+
+				CommitScrubPosition(NewValue, /*bIsScrubbing=*/true);
+			}
+		}
+		return FReply::Handled();
+	}
+
+	return FReply::Unhandled();
+}
+
+void FVisualLoggerTimeSliderController::CommitScrubPosition( float NewValue, bool bIsScrubbing )
+{
+	// Manage the scrub position ourselves if its not bound to a delegate
+	if ( !TimeSliderArgs.ScrubPosition.IsBound() )
+	{
+		TimeSliderArgs.ScrubPosition.Set( NewValue );
+	}
+
+	TRange<float> LocalViewRange = TimeSliderArgs.ViewRange.Get();
+	const float RangeSize = LocalViewRange.Size<float>();
+	if (NewValue < LocalViewRange.GetLowerBoundValue())
+	{
+		SetTimeRange(NewValue, NewValue + RangeSize);
+	}
+	else if (NewValue > LocalViewRange.GetUpperBoundValue())
+	{
+		SetTimeRange(NewValue - RangeSize, NewValue);
+	}
+
+	TimeSliderArgs.OnScrubPositionChanged.ExecuteIfBound( NewValue, bIsScrubbing );
+}
+
+void FVisualLoggerTimeSliderController::SetExternalScrollbar(TSharedRef<SScrollBar> InScrollbar) 
+{ 
+	Scrollbar = InScrollbar;
+	Scrollbar->SetOnUserScrolled(FOnUserScrolled::CreateRaw(this, &FVisualLoggerTimeSliderController::HorizontalScrollBar_OnUserScrolled)); 
+};
+
+void FVisualLoggerTimeSliderController::HorizontalScrollBar_OnUserScrolled(float ScrollOffset)
+{
+	if (!TimeSliderArgs.ViewRange.IsBound())
+	{
+		TRange<float> LocalViewRange = TimeSliderArgs.ViewRange.Get();
+		float LocalViewRangeMin = LocalViewRange.GetLowerBoundValue();
+		float LocalViewRangeMax = LocalViewRange.GetUpperBoundValue();
+		float LocalClampMin = TimeSliderArgs.ClampRange.Get().GetLowerBoundValue();
+		float LocalClampMax = TimeSliderArgs.ClampRange.Get().GetUpperBoundValue();
+
+		float InThumbSizeFraction = (LocalViewRangeMax - LocalViewRangeMin) / (LocalClampMax - LocalClampMin);
+
+		float NewViewOutputMin = LocalClampMin + ScrollOffset * (LocalClampMax - LocalClampMin);
+		// The  output is not bound to a delegate so we'll manage the value ourselves
+		float NewViewOutputMax = FMath::Min<float>(NewViewOutputMin + (LocalViewRangeMax - LocalViewRangeMin), LocalClampMax);
+		NewViewOutputMin = NewViewOutputMax - (LocalViewRangeMax - LocalViewRangeMin);
+
+		float InOffsetFraction = (NewViewOutputMin - LocalClampMin) / (LocalClampMax - LocalClampMin);
+		//if (InOffsetFraction + InThumbSizeFraction <= 1)
+		{
+			TimeSliderArgs.ViewRange.Set(TRange<float>(NewViewOutputMin, NewViewOutputMax));
+			Scrollbar->SetState(InOffsetFraction, InThumbSizeFraction);
+		}
+	}
+}
+
+void FVisualLoggerTimeSliderController::SetTimeRange(float NewViewOutputMin, float NewViewOutputMax)
+{
+	TimeSliderArgs.ViewRange.Set(TRange<float>(NewViewOutputMin, NewViewOutputMax));
+
+	float LocalClampMin = TimeSliderArgs.ClampRange.Get().GetLowerBoundValue();
+	float LocalClampMax = TimeSliderArgs.ClampRange.Get().GetUpperBoundValue();
+
+	const float InOffsetFraction = (NewViewOutputMin - LocalClampMin) / (LocalClampMax - LocalClampMin);
+	const float InThumbSizeFraction = (NewViewOutputMax - NewViewOutputMin) / (LocalClampMax - LocalClampMin);
+	Scrollbar->SetState(InOffsetFraction, InThumbSizeFraction);
+}
+
+void FVisualLoggerTimeSliderController::SetClampRange(float MinValue, float MaxValue)
+{
+	TRange<float> LocalViewRange = TimeSliderArgs.ViewRange.Get();
+	float LocalClampMin = TimeSliderArgs.ClampRange.Get().GetLowerBoundValue();
+	float LocalClampMax = TimeSliderArgs.ClampRange.Get().GetUpperBoundValue();
+	const float CurrentDistance = LocalClampMax - LocalClampMin;
+	const float ZoomDelta = (LocalViewRange.GetUpperBoundValue() - LocalViewRange.GetLowerBoundValue()) / CurrentDistance;
+
+	MaxValue = MinValue + (MaxValue - MinValue < 2 ? CurrentDistance : MaxValue - MinValue);
+
+	TimeSliderArgs.ClampRange = TRange<float>(MinValue, MaxValue);
+
+	const float LocalViewRangeMin = FMath::Clamp(LocalViewRange.GetLowerBoundValue(), MinValue, MaxValue);
+	const float LocalViewRangeMax = FMath::Clamp(LocalViewRange.GetUpperBoundValue(), MinValue, MaxValue);
+	SetTimeRange(ZoomDelta >= 1 ? MinValue : LocalViewRangeMin, ZoomDelta >= 1 ? MaxValue : LocalViewRangeMax);
+}
+
+FReply FVisualLoggerTimeSliderController::OnMouseWheel( SWidget& WidgetOwner, const FGeometry& MyGeometry, const FPointerEvent& MouseEvent )
+{
+	FReply ReturnValue = FReply::Unhandled();;
+
+	if (MouseEvent.IsLeftShiftDown())
+	{
+		const float ZoomDelta = 0.025f * MouseEvent.GetWheelDelta();
+		TimeSliderArgs.CursorSize.Set(FMath::Clamp(TimeSliderArgs.CursorSize.Get() + ZoomDelta, 0.0f, 1.0f));
+
+		ReturnValue = FReply::Handled();
+	}
+	else if ( TimeSliderArgs.AllowZoom )
+	{
+		const float ZoomDelta = -0.1f * MouseEvent.GetWheelDelta();
+
+		{
+			TRange<float> LocalViewRange = TimeSliderArgs.ViewRange.Get();
+			float LocalViewRangeMax = LocalViewRange.GetUpperBoundValue();
+			float LocalViewRangeMin = LocalViewRange.GetLowerBoundValue();
+			const float OutputViewSize = LocalViewRangeMax - LocalViewRangeMin;
+			const float OutputChange = OutputViewSize * ZoomDelta;
+
+			float NewViewOutputMin = LocalViewRangeMin - (OutputChange * 0.5f);
+			float NewViewOutputMax = LocalViewRangeMax + (OutputChange * 0.5f);
+
+			if( FMath::Abs( OutputChange ) > 0.01f && NewViewOutputMin < NewViewOutputMax )
+			{
+				float LocalClampMin = TimeSliderArgs.ClampRange.Get().GetLowerBoundValue();
+				float LocalClampMax = TimeSliderArgs.ClampRange.Get().GetUpperBoundValue();
+
+				// Clamp the range if clamp values are set
+				if ( NewViewOutputMin < LocalClampMin )
+				{
+					NewViewOutputMin = LocalClampMin;
+				}
+				
+				if ( NewViewOutputMax > LocalClampMax )
+				{
+					NewViewOutputMax = LocalClampMax;
+				}
+
+				TimeSliderArgs.OnViewRangeChanged.ExecuteIfBound(TRange<float>(NewViewOutputMin, NewViewOutputMax), EViewRangeInterpolation::Immediate);
+				if (Scrollbar.IsValid())
+				{
+					float InOffsetFraction = (NewViewOutputMin - LocalClampMin) / (LocalClampMax - LocalClampMin);
+					float InThumbSizeFraction = (NewViewOutputMax - NewViewOutputMin) / (LocalClampMax - LocalClampMin);
+					Scrollbar->SetState(InOffsetFraction, InThumbSizeFraction);
+				}
+				if( !TimeSliderArgs.ViewRange.IsBound() )
+				{	
+					// The  output is not bound to a delegate so we'll manage the value ourselves
+					TimeSliderArgs.ViewRange.Set( TRange<float>( NewViewOutputMin, NewViewOutputMax ) );
+				}
+			}
+		}
+
+		ReturnValue = FReply::Handled();
+	}
+
+	return ReturnValue;
+}
+
+int32 FVisualLoggerTimeSliderController::OnPaintSectionView( const FGeometry& AllottedGeometry, const FSlateRect& MyClippingRect, FSlateWindowElementList& OutDrawElements, int32 LayerId, bool bEnabled, bool bDisplayTickLines, bool bDisplayScrubPosition  ) const
+{
+	const ESlateDrawEffect::Type DrawEffects = bEnabled ? ESlateDrawEffect::None : ESlateDrawEffect::DisabledEffect;
+
+	TRange<float> LocalViewRange = TimeSliderArgs.ViewRange.Get();
+	float LocalScrubPosition = TimeSliderArgs.ScrubPosition.Get();
+
+	float ViewRange = LocalViewRange.Size<float>();
+	float PixelsPerInput = ViewRange > 0 ? AllottedGeometry.Size.X / ViewRange : 0;
+	float LinePos =  (LocalScrubPosition - LocalViewRange.GetLowerBoundValue()) * PixelsPerInput;
+
+	FScrubRangeToScreen RangeToScreen( LocalViewRange, AllottedGeometry.Size );
+
+	if( bDisplayTickLines )
+	{
+		// Draw major tick lines in the section area
+		FDrawTickArgs Args;
+		Args.AllottedGeometry = AllottedGeometry;
+		Args.bMirrorLabels = false;
+		Args.bOnlyDrawMajorTicks = true;
+		Args.TickColor = FLinearColor( 0.3f, 0.3f, 0.3f, 0.3f );
+		Args.ClippingRect = MyClippingRect;
+		Args.DrawEffects = DrawEffects;
+		// Draw major ticks under sections
+		Args.StartLayer = LayerId-1;
+		// Draw the tick the entire height of the section area
+		Args.TickOffset = 0.0f;
+		Args.MajorTickHeight = AllottedGeometry.Size.Y;
+
+		DrawTicks( OutDrawElements, RangeToScreen, Args );
+	}
+
+	if( bDisplayScrubPosition )
+	{
+		// Draw cursor size
+		const float CursorHalfSize = TimeSliderArgs.CursorSize.Get() * 0.5f;
+		const float CursorHalfLength = AllottedGeometry.Size.X * CursorHalfSize;
+		FPaintGeometry CursorGeometry = AllottedGeometry.ToPaintGeometry(FVector2D(LinePos - CursorHalfLength, 0), FVector2D(2 * CursorHalfLength, AllottedGeometry.Size.Y));
+
+		FSlateDrawElement::MakeBox(
+			OutDrawElements,
+			++LayerId,
+			CursorGeometry,
+			CursorBackground,
+			MyClippingRect,
+			DrawEffects,
+			FLinearColor::White.CopyWithNewOpacity(0.08f)
+			);
+
+		// Draw a line for the scrub position
+		TArray<FVector2D> LinePoints;
+		LinePoints.AddUninitialized(2);
+		LinePoints[0] = FVector2D( 1.0f, 0.0f );
+		LinePoints[1] = FVector2D( 1.0f, FMath::RoundToFloat( AllottedGeometry.Size.Y ) );
+
+		FSlateDrawElement::MakeLines(
+			OutDrawElements,
+			++LayerId,
+			AllottedGeometry.ToPaintGeometry( FVector2D(LinePos, 0.0f ), FVector2D(1.0f,1.0f) ),
+			LinePoints,
+			MyClippingRect,
+			DrawEffects,
+			FLinearColor::White.CopyWithNewOpacity(0.39f),
+			false
+			);
+
+	}
+
+	return LayerId;
+}
+
+#undef LOCTEXT_NAMESPACE
diff --git a/Engine/Source/Developer/LogVisualizer/Private/VisualLoggerTimeSliderController.h b/Engine/Source/Developer/LogVisualizer/Private/VisualLoggerTimeSliderController.h
new file mode 100644
--- /dev/null
+++ b/Engine/Source/Developer/LogVisualizer/Private/VisualLoggerTimeSliderController.h
@@ -0,0 +1,83 @@
+// Copyright 1998-2016 Epic Games, Inc. All Rights Reserved.
+
+#pragma once
+
+#include "Editor/SequencerWidgets/Public/ITimeSlider.h"
+
+struct FVisualLoggerTimeSliderArgs : FTimeSliderArgs
+{
+	FVisualLoggerTimeSliderArgs()
+		: CursorSize(0.05f)
+	{
+		ViewRange = TRange<float>(0.0f, 5.0f);
+	}
+
+	/** Cursor range for data like histogram graphs, etc. */
+	TAttribute< float > CursorSize;
+};
+
+/**
+ * A time slider controller for sequencer
+ * Draws and manages time data for a Sequencer
+ */
+class FVisualLoggerTimeSliderController : public ITimeSliderController
+{
+public:
+	FVisualLoggerTimeSliderController(const FVisualLoggerTimeSliderArgs& InArgs);
+	void SetTimesliderArgs(const FVisualLoggerTimeSliderArgs& InArgs);
+
+	/** ITimeSliderController Interface */
+	virtual int32 OnPaintTimeSlider( bool bMirrorLabels, const FGeometry& AllottedGeometry, const FSlateRect& MyClippingRect, FSlateWindowElementList& OutDrawElements, int32 LayerId, const FWidgetStyle& InWidgetStyle, bool bParentEnabled ) const override;
+	virtual FReply OnMouseButtonDown( SWidget& WidgetOwner, const FGeometry& MyGeometry, const FPointerEvent& MouseEvent ) override;
+	virtual FReply OnMouseButtonUp( SWidget& WidgetOwner, const FGeometry& MyGeometry, const FPointerEvent& MouseEvent ) override;
+	virtual FReply OnMouseMove( SWidget& WidgetOwner, const FGeometry& MyGeometry, const FPointerEvent& MouseEvent ) override;
+	virtual FReply OnMouseWheel( SWidget& WidgetOwner, const FGeometry& MyGeometry, const FPointerEvent& MouseEvent ) override;
+	virtual FCursorReply OnCursorQuery( TSharedRef<const SWidget> WidgetOwner, const FGeometry& MyGeometry, const FPointerEvent& CursorEvent) const override { return FCursorReply::Unhandled(); }
+
+	/**
+	 * Draws major tick lines in the section view                                                              
+	 */
+	int32 OnPaintSectionView( const FGeometry& AllottedGeometry, const FSlateRect& MyClippingRect, FSlateWindowElementList& OutDrawElements, int32 LayerId, bool bEnabled, bool bDisplayTickLines, bool bDisplayScrubPosition ) const;
+	FVisualLoggerTimeSliderArgs& GetTimeSliderArgs() { return TimeSliderArgs; }
+	/**
+	 * Call this method when the user's interaction has changed the scrub position
+	 *
+	 * @param NewValue				Value resulting from the user's interaction
+	 * @param bIsScrubbing			True if done via scrubbing, false if just releasing scrubbing
+	 */
+	void CommitScrubPosition( float NewValue, bool bIsScrubbing );
+
+	void SetExternalScrollbar(TSharedRef<SScrollBar> Scrollbar);
+	void SetTimeRange(float MinValue, float MaxValue);
+	void SetClampRange(float MinValue, float MaxValue);
+	bool IsPanning() { return bPanning; }
+
+private:
+
+	/**
+	 * Draws time tick marks
+	 *
+	 * @param OutDrawElements	List to add draw elements to
+	 * @param RangeToScreen		Time range to screen space converter
+	 * @param InArgs			Parameters for drawing the tick lines
+	 */
+	void DrawTicks( FSlateWindowElementList& OutDrawElements, const struct FScrubRangeToScreen& RangeToScreen, struct FDrawTickArgs& InArgs ) const;
+	void HorizontalScrollBar_OnUserScrolled(float ScrollOffset);
+private:
+	FVisualLoggerTimeSliderArgs TimeSliderArgs;
+	/** Brush for drawing an upwards facing scrub handle */
+	const FSlateBrush* ScrubHandleUp;
+	/** Brush for drawing a downwards facing scrub handle */
+	const FSlateBrush* ScrubHandleDown;
+	/** Brush for drawing cursor background to visualize corsor size */
+	const FSlateBrush* CursorBackground;
+	/** Total mouse delta during dragging **/
+	float DistanceDragged;
+	/** If we are dragging the scrubber */
+	bool bDraggingScrubber;
+	/** If we are currently panning the panel */
+	bool bPanning;
+	/***/
+	TSharedPtr<SScrollBar> Scrollbar;
+	FVector2D SoftwareCursorPosition;
+};
