Fix missing style warning opening the visual logger

diff --git a/Engine/Source/Developer/LogVisualizer/Private/LogVisualizerStyle.cpp b/Engine/Source/Developer/LogVisualizer/Private/LogVisualizerStyle.cpp
--- a/Engine/Source/Developer/LogVisualizer/Private/LogVisualizerStyle.cpp
+++ b/Engine/Source/Developer/LogVisualizer/Private/LogVisualizerStyle.cpp
@@ -115,6 +115,12 @@ TSharedRef< FSlateStyleSet > FLogVisualizerStyle::Create()
 			);
 	}
 
+	Style.Set("GenericFilters.TextStyle", FTextBlockStyle(NormalText)
+		.SetFont(DEFAULT_FONT("Bold", 9))
+		.SetColorAndOpacity(FLinearColor(1.0f, 1.0f, 1.0f, 0.9f))
+		.SetShadowOffset(FVector2D(1, 1))
+		.SetShadowColorAndOpacity(FLinearColor(0, 0, 0, 0.9f)));
+
 	// Toolbar
 	{
 		Style.Set("Toolbar.Pause", new IMAGE_BRUSH("Icons/icon_pause_40x", Icon40x40));
