Added tooltip support to the Visual Logger showing the time and all logs (both normal and shape descriptions) for the closest point in a timeline track

diff --git a/Engine/Source/Developer/LogVisualizer/Private/SVisualLoggerTimeline.cpp b/Engine/Source/Developer/LogVisualizer/Private/SVisualLoggerTimeline.cpp
--- a/Engine/Source/Developer/LogVisualizer/Private/SVisualLoggerTimeline.cpp
+++ b/Engine/Source/Developer/LogVisualizer/Private/SVisualLoggerTimeline.cpp
@@ -180,7 +180,7 @@ void SLogVisualizerTimeline::Construct(const FArguments& InArgs, TSharedPtr<FVis
 							[
 								SNew(STextBlock)
 								.Text(FText::FromName(OwnerClassName))
-								.TextStyle(FLogVisualizerStyle::Get(), TEXT("Sequencer.ClassNAme"))
+								.TextStyle(FLogVisualizerStyle::Get(), TEXT("Sequencer.ClassName"))
 							]
 						]
 					]
diff --git a/Engine/Source/Developer/LogVisualizer/Private/SVisualLoggerTimelineBar.cpp b/Engine/Source/Developer/LogVisualizer/Private/SVisualLoggerTimelineBar.cpp
--- a/Engine/Source/Developer/LogVisualizer/Private/SVisualLoggerTimelineBar.cpp
+++ b/Engine/Source/Developer/LogVisualizer/Private/SVisualLoggerTimelineBar.cpp
@@ -7,6 +7,7 @@
 #include "LogVisualizerStyle.h"
 #include "LogVisualizerPrivate.h"
 #include "VisualLoggerTimeSliderController.h"
+#include "Misc/OutputDeviceHelper.h"
 
 FReply SVisualLoggerTimelineBar::OnMouseButtonDown(const FGeometry& MyGeometry, const FPointerEvent& MouseEvent)
 {
@@ -49,6 +50,39 @@ FReply SVisualLoggerTimelineBar::OnMouseButtonUp(const FGeometry& MyGeometry, co
 
 FReply SVisualLoggerTimelineBar::OnMouseMove(const FGeometry& MyGeometry, const FPointerEvent& MouseEvent)
 {
+	FName RowName = TimelineOwner.Pin()->GetName();
+	FVisualLoggerDBRow& DBRow = FVisualLoggerDatabase::Get().GetRowByName(RowName);
+
+	const float ClosestMouseTime = TimeSliderController->GetTimeAtCursorPosition(MyGeometry, MouseEvent);
+
+	const int32 NewItemIndex = DBRow.GetClosestItem(ClosestMouseTime);
+	const auto& Items = DBRow.GetItems();
+
+	if (NewItemIndex != MouseMoveClosestItemIndex)
+	{
+		MouseMoveClosestItemIndex = NewItemIndex;
+
+		FString TooltipBuilder;
+		if (Items.IsValidIndex(MouseMoveClosestItemIndex))
+		{
+			const FVisualLogEntry& CurrentEntry = Items[MouseMoveClosestItemIndex].Entry;
+				
+			TooltipBuilder = FString::Printf(TEXT("Time: %.2f"), CurrentEntry.TimeStamp);
+
+			for (const FVisualLogShapeElement& Shape : CurrentEntry.ElementsToDraw)
+			{
+				TooltipBuilder += FString::Printf(TEXT("\n(shape) %s: %s"), FOutputDeviceHelper::VerbosityToString(Shape.Verbosity), *Shape.Description);
+			}
+
+			for (const FVisualLogLine& Line : CurrentEntry.LogLines)
+			{
+				TooltipBuilder += FString::Printf(TEXT("\n(log) %s: %s"), FOutputDeviceHelper::VerbosityToString(Line.Verbosity), *Line.Line);
+			}
+		}
+
+		SetToolTipText(FText::AsCultureInvariant(TooltipBuilder));
+	}
+
 	return TimeSliderController->OnMouseMove(*this, MyGeometry, MouseEvent);
 }
 
@@ -76,6 +110,8 @@ void SVisualLoggerTimelineBar::Construct(const FArguments& InArgs, TSharedPtr<FV
 	TimelineOwner = InTimelineOwner;
 
 	TRange<float> LocalViewRange = TimeSliderController->GetTimeSliderArgs().ViewRange.Get();
+
+	MouseMoveClosestItemIndex = INDEX_NONE;
 }
 
 FVector2D SVisualLoggerTimelineBar::ComputeDesiredSize(float) const
diff --git a/Engine/Source/Developer/LogVisualizer/Private/SVisualLoggerTimelineBar.h b/Engine/Source/Developer/LogVisualizer/Private/SVisualLoggerTimelineBar.h
--- a/Engine/Source/Developer/LogVisualizer/Private/SVisualLoggerTimelineBar.h
+++ b/Engine/Source/Developer/LogVisualizer/Private/SVisualLoggerTimelineBar.h
@@ -38,4 +38,5 @@ public:
 protected:
 	TSharedPtr<class FVisualLoggerTimeSliderController> TimeSliderController;
 	TWeakPtr<class SLogVisualizerTimeline> TimelineOwner;
+	int32 MouseMoveClosestItemIndex;
 };
diff --git a/Engine/Source/Developer/LogVisualizer/Private/VisualLoggerTimeSliderController.cpp b/Engine/Source/Developer/LogVisualizer/Private/VisualLoggerTimeSliderController.cpp
--- a/Engine/Source/Developer/LogVisualizer/Private/VisualLoggerTimeSliderController.cpp
+++ b/Engine/Source/Developer/LogVisualizer/Private/VisualLoggerTimeSliderController.cpp
@@ -368,6 +368,28 @@ FReply FVisualLoggerTimeSliderController::OnMouseButtonUp( SWidget& WidgetOwner,
 	return FReply::Unhandled();
 }
 
+float FVisualLoggerTimeSliderController::GetTimeAtCursorPosition(const FGeometry& MyGeometry, const FPointerEvent& MouseEvent) const
+{
+	FScrubRangeToScreen RangeToScreen(TimeSliderArgs.ViewRange.Get(), MyGeometry.GetLocalSize());
+	FVector2D CursorPos = MyGeometry.AbsoluteToLocal(MouseEvent.GetLastScreenSpacePosition());
+	float NewValue = RangeToScreen.LocalXToInput(CursorPos.X);
+
+	float LocalClampMin = TimeSliderArgs.ClampRange.Get().GetLowerBoundValue();
+	float LocalClampMax = TimeSliderArgs.ClampRange.Get().GetUpperBoundValue();
+
+	if (NewValue < LocalClampMin)
+	{
+		NewValue = LocalClampMin;
+	}
+
+	if (NewValue > LocalClampMax)
+	{
+		NewValue = LocalClampMax;
+	}
+
+	return NewValue;
+}
+
 FReply FVisualLoggerTimeSliderController::OnMouseMove( SWidget& WidgetOwner, const FGeometry& MyGeometry, const FPointerEvent& MouseEvent )
 {
 	if ( WidgetOwner.HasMouseCapture() )
@@ -442,23 +464,7 @@ FReply FVisualLoggerTimeSliderController::OnMouseMove( SWidget& WidgetOwner, con
 			}
 			else
 			{
-				FScrubRangeToScreen RangeToScreen( TimeSliderArgs.ViewRange.Get(), MyGeometry.GetLocalSize());
-				FVector2D CursorPos = MyGeometry.AbsoluteToLocal( MouseEvent.GetLastScreenSpacePosition() );
-				float NewValue = RangeToScreen.LocalXToInput( CursorPos.X );
-
-				float LocalClampMin = TimeSliderArgs.ClampRange.Get().GetLowerBoundValue();
-				float LocalClampMax = TimeSliderArgs.ClampRange.Get().GetUpperBoundValue();
-
-				if (NewValue < LocalClampMin)
-				{
-					NewValue = LocalClampMin;
-				}
-
-				if (NewValue > LocalClampMax)
-				{
-					NewValue = LocalClampMax;
-				}
-
+				const float NewValue = GetTimeAtCursorPosition(MyGeometry, MouseEvent);
 				CommitScrubPosition(NewValue, /*bIsScrubbing=*/true);
 			}
 		}
diff --git a/Engine/Source/Developer/LogVisualizer/Private/VisualLoggerTimeSliderController.h b/Engine/Source/Developer/LogVisualizer/Private/VisualLoggerTimeSliderController.h
--- a/Engine/Source/Developer/LogVisualizer/Private/VisualLoggerTimeSliderController.h
+++ b/Engine/Source/Developer/LogVisualizer/Private/VisualLoggerTimeSliderController.h
@@ -101,6 +101,7 @@ public:
 	void SetClampRange(float MinValue, float MaxValue);
 	bool IsPanning() { return bPanning; }
 
+	float GetTimeAtCursorPosition(const FGeometry& MyGeometry, const FPointerEvent& MouseEvent) const;
 private:
 	// forward declared as class members to prevent name collision with similar types defined in other units
 	struct FScrubRangeToScreen;
