PR #6127: Fix histogram render so each line is a different color. (Contributed by [REDACTED_PERSON])

Details: The golden ration hue trick is not utilized because of the 0 saturation and all the lines are white in VisualLogger graphs.

diff --git a/Engine/Source/Developer/LogVisualizer/Private/VisualLoggerCanvasRenderer.cpp b/Engine/Source/Developer/LogVisualizer/Private/VisualLoggerCanvasRenderer.cpp
--- a/Engine/Source/Developer/LogVisualizer/Private/VisualLoggerCanvasRenderer.cpp
+++ b/Engine/Source/Developer/LogVisualizer/Private/VisualLoggerCanvasRenderer.cpp
@@ -269,7 +269,7 @@ void FVisualLoggerCanvasRenderer::DrawHistogramGraphs(class UCanvas* Canvas, cla
 					Hue -= FMath::FloorToFloat(Hue);
 				}
 
-				HistogramGraph->GetGraphLine(LineIndex)->Color = FLinearColor::MakeFromHSV8(Hue * 255, 0, 244);
+				HistogramGraph->GetGraphLine(LineIndex)->Color = FLinearColor::MakeFromHSV8(Hue * 255, 200, 244);
 				HistogramGraph->GetGraphLine(LineIndex)->LineName = DataName;
 				HistogramGraph->GetGraphLine(LineIndex)->Data.Append(LinesIt->Value.Samples);
 				HistogramGraph->GetGraphLine(LineIndex)->LeftExtreme = LinesIt->Value.LeftExtreme;
