Add option to recompute graph min/max Y based on data samples that will be rendered for a given frame. Useful for data that grows continually.

#rnx

diff --git a/Engine/Source/Developer/LogVisualizer/Private/LogVisualizerSettings.cpp b/Engine/Source/Developer/LogVisualizer/Private/LogVisualizerSettings.cpp
--- a/Engine/Source/Developer/LogVisualizer/Private/LogVisualizerSettings.cpp
+++ b/Engine/Source/Developer/LogVisualizer/Private/LogVisualizerSettings.cpp
@@ -19,6 +19,7 @@ ULogVisualizerSettings::ULogVisualizerSettings(const FObjectInitializer& ObjectI
 	bResetDataWithNewSession = false;
 	bDrawExtremesOnGraphs = false;
 	bUsePlayersOnlyForPause = true;
+	bConstrainGraphToLocalMinMax = false;
 }
 
 class UMaterial* ULogVisualizerSettings::GetDebugMeshMaterial()
diff --git a/Engine/Source/Developer/LogVisualizer/Private/VisualLoggerCanvasRenderer.cpp b/Engine/Source/Developer/LogVisualizer/Private/VisualLoggerCanvasRenderer.cpp
--- a/Engine/Source/Developer/LogVisualizer/Private/VisualLoggerCanvasRenderer.cpp
+++ b/Engine/Source/Developer/LogVisualizer/Private/VisualLoggerCanvasRenderer.cpp
@@ -239,6 +239,7 @@ void FVisualLoggerCanvasRenderer::DrawHistogramGraphs(class UCanvas* Canvas, cla
 		int32 CurrentColumn = 0;
 		int32 CurrentRow = 0;
 		bool bDrawExtremesOnGraphs = ULogVisualizerSettings::StaticClass()->GetDefaultObject<ULogVisualizerSettings>()->bDrawExtremesOnGraphs;
+		const bool bConstrainGraphToLocalMinMax = ULogVisualizerSettings::StaticClass()->GetDefaultObject<ULogVisualizerSettings>()->bConstrainGraphToLocalMinMax;
 		for (auto It(CollectedGraphs.CreateIterator()); It; ++It)
 		{
 			TWeakObjectPtr<UReporterGraph> HistogramGraph = Canvas->GetReporterGraph();
@@ -255,6 +256,13 @@ void FVisualLoggerCanvasRenderer::DrawHistogramGraphs(class UCanvas* Canvas, cla
 			auto& CategoriesForGraph = UsedGraphCategories.FindOrAdd(It->Key.ToString());
 
 			It->Value.GraphLines.KeySort(FNameLexicalLess());
+
+			if (bConstrainGraphToLocalMinMax)
+			{
+				It->Value.Min.Y = FLT_MAX;
+				It->Value.Max.Y = -FLT_MAX;
+			}
+
 			for (auto LinesIt(It->Value.GraphLines.CreateConstIterator()); LinesIt; ++LinesIt)
 			{
 				const FString DataName = LinesIt->Value.DataName.ToString();
@@ -280,6 +288,18 @@ void FVisualLoggerCanvasRenderer::DrawHistogramGraphs(class UCanvas* Canvas, cla
 				MaxStringSize = StringSizeX > MaxStringSize ? StringSizeX : MaxStringSize;
 
 				++LineIndex;
+
+				if (bConstrainGraphToLocalMinMax)
+				{
+					for (const FVector2D& SampleData : LinesIt->Value.Samples)
+					{
+						if (SampleData.X >= TimeStampWindow.X && SampleData.X <= TimeStampWindow.Y)
+						{
+							It->Value.Min.Y = FMath::Min<float>(It->Value.Min.Y, SampleData.Y);
+							It->Value.Max.Y = FMath::Max<float>(It->Value.Max.Y, SampleData.Y);
+						}
+					}
+				}
 			}
 
 			FVector2D GraphSpaceSize;
diff --git a/Engine/Source/Developer/LogVisualizer/Public/LogVisualizerSettings.h b/Engine/Source/Developer/LogVisualizer/Public/LogVisualizerSettings.h
--- a/Engine/Source/Developer/LogVisualizer/Public/LogVisualizerSettings.h
+++ b/Engine/Source/Developer/LogVisualizer/Public/LogVisualizerSettings.h
@@ -150,6 +150,10 @@ public:
 	UPROPERTY(EditAnywhere, config, Category = "VisualLogger")
 	bool bDrawExtremesOnGraphs;
 
+	/** Graphs will be scaled around local Min/Max values (values being displayed) rather than all historic data */
+	UPROPERTY(EditAnywhere, config, Category = "VisualLogger")
+	bool bConstrainGraphToLocalMinMax;
+
 	/** Whether to use PlayersOnly during Pause or not */
 	UPROPERTY(EditAnywhere, config, Category = "VisualLogger")
 	bool bUsePlayersOnlyForPause;
