Fixed a VisualLogger ensure by allowing FVisualLoggerEditorInterface::GetHelperActor to spawn during construction script (helper actor is created on demand and only once per world).

#rnx

diff --git a/Engine/Source/Developer/LogVisualizer/Private/VisualLoggerDatabase.cpp b/Engine/Source/Developer/LogVisualizer/Private/VisualLoggerDatabase.cpp
--- a/Engine/Source/Developer/LogVisualizer/Private/VisualLoggerDatabase.cpp
+++ b/Engine/Source/Developer/LogVisualizer/Private/VisualLoggerDatabase.cpp
@@ -471,6 +471,10 @@ AActor* FVisualLoggerEditorInterface::GetHelperActor(UWorld* InWorld) const
 	FActorSpawnParameters SpawnInfo;
 	SpawnInfo.SpawnCollisionHandlingOverride = ESpawnActorCollisionHandlingMethod::AlwaysSpawn;
 	SpawnInfo.Name = *FString::Printf(TEXT("VisualLoggerRenderingActor"));
+
+	// The helper actor is created on demand and only once per world so we can allow it to spawn during construction script.
+	SpawnInfo.bAllowDuringConstructionScript = true;
+
 	return World->SpawnActor<AVisualLoggerRenderingActor>(SpawnInfo);
 }
 
