Fixed vislog's perf and architecture issue resulting in vislog constantly recreating AVisualLoggerRenderingActor in client-server PIE due to vlogging with different worlds.

diff --git a/Engine/Source/Developer/LogVisualizer/Private/SVisualLogger.cpp b/Engine/Source/Developer/LogVisualizer/Private/SVisualLogger.cpp
--- a/Engine/Source/Developer/LogVisualizer/Private/SVisualLogger.cpp
+++ b/Engine/Source/Developer/LogVisualizer/Private/SVisualLogger.cpp
@@ -94,30 +94,18 @@ SVisualLogger::SVisualLogger()
 	class FVisualLoggerDevice : public FVisualLogDevice
 	{
 	public:
-		FVisualLoggerDevice(SVisualLogger* InVisualLogger, FOnWorldChanged OnWorldChangedDelegate) : VisualLoggerWidget(InVisualLogger), LastUsedWorld(nullptr), OnWorldChanged(OnWorldChangedDelegate) {}
+		FVisualLoggerDevice(SVisualLogger& InVisualLogger) : VisualLoggerWidget(InVisualLogger) {}
 
 		virtual ~FVisualLoggerDevice(){}
 		virtual void Serialize(const UObject* LogOwner, FName OwnerName, FName OwnerClassName, const FVisualLogEntry& LogEntry) override
 		{
-			VisualLoggerWidget->OnNewLogEntry(FVisualLogDevice::FVisualLogEntryItem(OwnerName, OwnerClassName, LogEntry));
-			UWorld* CurrentWorld = nullptr;
-			if (FVisualLogger::Get().GetObjectToWorldMap().Contains(LogOwner))
-			{
-				CurrentWorld = const_cast<UWorld*>(FVisualLogger::Get().GetObjectToWorldMap()[LogOwner].Get());
-			}
-			if (LastUsedWorld != CurrentWorld && CurrentWorld != nullptr)
-			{
-				OnWorldChanged.ExecuteIfBound(LastUsedWorld, CurrentWorld);
-				LastUsedWorld = CurrentWorld;
-			}
+			VisualLoggerWidget.OnNewLogEntry(FVisualLogDevice::FVisualLogEntryItem(OwnerName, OwnerClassName, LogEntry));
 		}
 		
-		SVisualLogger* VisualLoggerWidget;
-		UWorld* LastUsedWorld;
-		FOnWorldChanged OnWorldChanged;
+		SVisualLogger& VisualLoggerWidget;
 	};
 
-	InternalDevice = MakeShareable(new FVisualLoggerDevice(this, FOnWorldChanged::CreateLambda([this](UWorld* PreviousWorld, UWorld* CurrentWorld){ OnNewWorld(CurrentWorld); })));
+	InternalDevice = MakeShareable(new FVisualLoggerDevice(*this));
 	FVisualLogger::Get().AddDevice(InternalDevice.Get());
 }
 
@@ -177,6 +165,7 @@ void SVisualLogger::Construct(const FArguments& InArgs, const TSharedRef<SDockTa
 	FVisualLoggerDatabase::Get().GetEvents().OnItemSelectionChanged.AddRaw(this, &SVisualLogger::OnItemsSelectionChanged);
 
 	GEngine->OnWorldAdded().AddRaw(this, &SVisualLogger::OnNewWorld);
+	GEngine->OnWorldAdded().AddRaw(this, &SVisualLogger::OnWorldDestroyed);
 
 	//////////////////////////////////////////////////////////////////////////
 	// Command Action Lists
@@ -737,15 +726,19 @@ void SVisualLogger::ResetData()
 	FLogVisualizer::Get().GetTimeSliderController().Get()->GetTimeSliderArgs().OnScrubPositionChanged = FVisualLoggerTimeSliderArgs::FOnScrubPositionChanged::CreateRaw(this, &SVisualLogger::OnScrubPositionChanged);
 }
 
-void SVisualLogger::OnNewWorld(UWorld* NewWorld)
+void SVisualLogger::OnWorldDestroyed(UWorld* NewWorld)
 {
-	if (LastUsedWorld.IsValid() && LastUsedWorld != NewWorld)
+	if (NewWorld)
 	{
-		for (TActorIterator<AVisualLoggerRenderingActor> It(LastUsedWorld.Get()); It; ++It)
+		for (TActorIterator<AVisualLoggerRenderingActor> It(NewWorld); It; ++It)
 		{
-			LastUsedWorld->DestroyActor(*It);
+			NewWorld->DestroyActor(*It);
 		}
 	}
+}
+
+void SVisualLogger::OnNewWorld(UWorld* NewWorld)
+{
 	LastUsedWorld = NewWorld;
 
 	AVisualLoggerRenderingActor* HelperActor = Cast<AVisualLoggerRenderingActor>(FVisualLoggerEditorInterface::Get()->GetHelperActor(LastUsedWorld.Get()));
diff --git a/Engine/Source/Developer/LogVisualizer/Private/SVisualLogger.h b/Engine/Source/Developer/LogVisualizer/Private/SVisualLogger.h
--- a/Engine/Source/Developer/LogVisualizer/Private/SVisualLogger.h
+++ b/Engine/Source/Developer/LogVisualizer/Private/SVisualLogger.h
@@ -102,6 +102,7 @@ public:
 	void HandleTabManagerPersistLayout(const TSharedRef<FTabManager::FLayout>& LayoutToSave);
 
 	void OnNewWorld(UWorld* NewWorld);
+	void OnWorldDestroyed(UWorld* InWorld);
 	void ResetData();
 
 protected:
diff --git a/Engine/Source/Developer/LogVisualizer/Private/VisualLoggerDatabase.cpp b/Engine/Source/Developer/LogVisualizer/Private/VisualLoggerDatabase.cpp
--- a/Engine/Source/Developer/LogVisualizer/Private/VisualLoggerDatabase.cpp
+++ b/Engine/Source/Developer/LogVisualizer/Private/VisualLoggerDatabase.cpp
@@ -470,7 +470,6 @@ AActor* FVisualLoggerEditorInterface::GetHelperActor(UWorld* InWorld) const
 
 	FActorSpawnParameters SpawnInfo;
 	SpawnInfo.SpawnCollisionHandlingOverride = ESpawnActorCollisionHandlingMethod::AlwaysSpawn;
-	SpawnInfo.Name = *FString::Printf(TEXT("VisualLoggerRenderingActor"));
 
 	// The helper actor is created on demand and only once per world so we can allow it to spawn during construction script.
 	SpawnInfo.bAllowDuringConstructionScript = true;
