PR #7408: [VisualLogger] Moved FLogVisualiser to a public header, added function to seek to a time (Contributed by [REDACTED_PERSON])

diff --git a/Engine/Source/Developer/LogVisualizer/Private/LogVisualizer.cpp b/Engine/Source/Developer/LogVisualizer/Private/LogVisualizer.cpp
--- a/Engine/Source/Developer/LogVisualizer/Private/LogVisualizer.cpp
+++ b/Engine/Source/Developer/LogVisualizer/Private/LogVisualizer.cpp
@@ -1,6 +1,7 @@
 // Copyright Epic Games, Inc. All Rights Reserved.
 
 #include "LogVisualizer.h"
+#include "LogVisualizerPublic.h"
 #include "EngineGlobals.h"
 #include "GameFramework/SpectatorPawn.h"
 #include "Engine/LocalPlayer.h"
@@ -8,6 +9,7 @@
 #include "LogVisualizerSettings.h"
 #include "VisualLoggerDatabase.h"
 #include "VisualLoggerCameraController.h"
+
 #if WITH_EDITOR
 #include "Editor/EditorEngine.h"
 #include "EditorViewportClient.h"
@@ -332,3 +334,7 @@ void FLogVisualizer::GotoLastItem(FName RowName)
 	}
 }
 
+void FLogVisualizer::SeekToTime(float Time)
+{
+	GetTimeSliderController()->CommitScrubPosition(Time,  /*bIsScrubbing=*/true);
+}
\ No newline at end of file
diff --git a/Engine/Source/Developer/LogVisualizer/Private/LogVisualizer.h b/Engine/Source/Developer/LogVisualizer/Private/LogVisualizer.h
--- a/Engine/Source/Developer/LogVisualizer/Private/LogVisualizer.h
+++ b/Engine/Source/Developer/LogVisualizer/Private/LogVisualizer.h
@@ -23,7 +23,6 @@
 #include "HAL/RunnableThread.h"
 #include "Engine/World.h"
 #include "SVisualLogger.h"
-#include "LogVisualizerPrivate.h"
 #include "Developer/LogVisualizer/Private/SVisualLoggerToolbar.h"
 #include "Developer/LogVisualizer/Private/SVisualLoggerFilters.h"
 #include "Widgets/Layout/SScrollBarTrack.h"
diff --git a/Engine/Source/Developer/LogVisualizer/Private/LogVisualizerModule.cpp b/Engine/Source/Developer/LogVisualizer/Private/LogVisualizerModule.cpp
--- a/Engine/Source/Developer/LogVisualizer/Private/LogVisualizerModule.cpp
+++ b/Engine/Source/Developer/LogVisualizer/Private/LogVisualizerModule.cpp
@@ -10,10 +10,12 @@
 #include "VisualLoggerDatabase.h"
 #include "LogVisualizerStyle.h"
 #include "SVisualLogger.h"
+#include "SVisualLoggerTab.h"
 #include "VisualLoggerCommands.h"
 #include "Widgets/Docking/SDockTab.h"
-#include "LogVisualizerPrivate.h"
+#include "LogVisualizerPublic.h"
 #include "Features/IModularFeatures.h"
+
 #if WITH_EDITOR
 #include "ISettingsModule.h"
 #include "WorkspaceMenuStructure.h"
diff --git a/Engine/Source/Developer/LogVisualizer/Private/LogVisualizerPrivate.h b/Engine/Source/Developer/LogVisualizer/Private/LogVisualizerPrivate.h
deleted file mode 100644
diff --git a/Engine/Source/Developer/LogVisualizer/Private/SVisualLogger.cpp b/Engine/Source/Developer/LogVisualizer/Private/SVisualLogger.cpp
--- a/Engine/Source/Developer/LogVisualizer/Private/SVisualLogger.cpp
+++ b/Engine/Source/Developer/LogVisualizer/Private/SVisualLogger.cpp
@@ -19,12 +19,12 @@
 #include "LogVisualizerStyle.h"
 #include "VisualLoggerCommands.h"
 #include "Widgets/Docking/SDockTab.h"
-#include "LogVisualizerPrivate.h"
 #include "SVisualLoggerToolbar.h"
 #include "SVisualLoggerFilters.h"
 #include "SVisualLoggerView.h"
 #include "SVisualLoggerLogsList.h"
 #include "SVisualLoggerStatusView.h"
+#include "SVisualLoggerTab.h"
 #include "VisualLoggerTimeSliderController.h"
 #include "VisualLoggerRenderingActor.h"
 #include "VisualLoggerCanvasRenderer.h"
diff --git a/Engine/Source/Developer/LogVisualizer/Private/SVisualLoggerBaseWidget.h b/Engine/Source/Developer/LogVisualizer/Private/SVisualLoggerBaseWidget.h
new file mode 100644
--- /dev/null
+++ b/Engine/Source/Developer/LogVisualizer/Private/SVisualLoggerBaseWidget.h
@@ -0,0 +1,21 @@
+// Copyright Epic Games, Inc. All Rights Reserved.
+#pragma once
+
+/* Dependencies
+*****************************************************************************/
+
+#include "CoreMinimal.h"
+#include "Input/Reply.h"
+#include "Widgets/SCompoundWidget.h"
+#include "Developer/LogVisualizer/Private/SVisualLogger.h"
+#include "LogVisualizerPublic.h"
+
+class SVisualLoggerBaseWidget : public SCompoundWidget
+{
+public:
+	virtual bool SupportsKeyboardFocus() const override { return true; }
+	virtual FReply OnKeyDown(const FGeometry& MyGeometry, const FKeyEvent& InKeyEvent) override
+	{
+		return FLogVisualizer::Get().GetEvents().OnKeyboardEvent.Execute(MyGeometry, InKeyEvent);
+	}
+};
diff --git a/Engine/Source/Developer/LogVisualizer/Private/SVisualLoggerFilters.cpp b/Engine/Source/Developer/LogVisualizer/Private/SVisualLoggerFilters.cpp
--- a/Engine/Source/Developer/LogVisualizer/Private/SVisualLoggerFilters.cpp
+++ b/Engine/Source/Developer/LogVisualizer/Private/SVisualLoggerFilters.cpp
@@ -13,6 +13,8 @@
 #include "LogVisualizerStyle.h"
 #include "SFilterWidget.h"
 #include "Widgets/Input/SSearchBox.h"
+#include "LogVisualizerPublic.h"
+
 #if WITH_EDITOR
 #include "Editor/EditorEngine.h"
 #include "EditorViewportClient.h"
diff --git a/Engine/Source/Developer/LogVisualizer/Private/SVisualLoggerFilters.h b/Engine/Source/Developer/LogVisualizer/Private/SVisualLoggerFilters.h
--- a/Engine/Source/Developer/LogVisualizer/Private/SVisualLoggerFilters.h
+++ b/Engine/Source/Developer/LogVisualizer/Private/SVisualLoggerFilters.h
@@ -3,10 +3,10 @@
 #pragma once
 
 #include "CoreMinimal.h"
+#include "SVisualLoggerBaseWidget.h"
 #include "Widgets/SWidget.h"
 #include "Widgets/DeclarativeSyntaxSupport.h"
 #include "Framework/Commands/UICommandList.h"
-#include "Developer/LogVisualizer/Private/LogVisualizerPrivate.h"
 
 class FMenuBuilder;
 class SComboButton;
diff --git a/Engine/Source/Developer/LogVisualizer/Private/SVisualLoggerLogsList.cpp b/Engine/Source/Developer/LogVisualizer/Private/SVisualLoggerLogsList.cpp
--- a/Engine/Source/Developer/LogVisualizer/Private/SVisualLoggerLogsList.cpp
+++ b/Engine/Source/Developer/LogVisualizer/Private/SVisualLoggerLogsList.cpp
@@ -8,6 +8,7 @@
 #include "LogVisualizerStyle.h"
 #include "Widgets/Views/SListView.h"
 #include "HAL/PlatformApplicationMisc.h"
+#include "LogVisualizerPublic.h"
 
 #define LOCTEXT_NAMESPACE "SVisualLoggerLogsList"
 
diff --git a/Engine/Source/Developer/LogVisualizer/Private/SVisualLoggerLogsList.h b/Engine/Source/Developer/LogVisualizer/Private/SVisualLoggerLogsList.h
--- a/Engine/Source/Developer/LogVisualizer/Private/SVisualLoggerLogsList.h
+++ b/Engine/Source/Developer/LogVisualizer/Private/SVisualLoggerLogsList.h
@@ -7,9 +7,9 @@
 #include "Widgets/DeclarativeSyntaxSupport.h"
 #include "Framework/Commands/UICommandList.h"
 #include "VisualLogger/VisualLoggerTypes.h"
-#include "Developer/LogVisualizer/Private/LogVisualizerPrivate.h"
 #include "Widgets/Views/STableViewBase.h"
 #include "Widgets/Views/STableRow.h"
+#include "SVisualLoggerBaseWidget.h"
 
 struct FLogEntryItem;
 struct FVisualLoggerDBRow;
diff --git a/Engine/Source/Developer/LogVisualizer/Private/SVisualLoggerStatusView.h b/Engine/Source/Developer/LogVisualizer/Private/SVisualLoggerStatusView.h
--- a/Engine/Source/Developer/LogVisualizer/Private/SVisualLoggerStatusView.h
+++ b/Engine/Source/Developer/LogVisualizer/Private/SVisualLoggerStatusView.h
@@ -6,7 +6,6 @@
 #include "Widgets/DeclarativeSyntaxSupport.h"
 #include "Framework/Commands/UICommandList.h"
 #include "VisualLogger/VisualLoggerTypes.h"
-#include "Developer/LogVisualizer/Private/LogVisualizerPrivate.h"
 #include "Widgets/Views/STableViewBase.h"
 #include "Widgets/Views/STableRow.h"
 #include "Widgets/Views/STreeView.h"
diff --git a/Engine/Source/Developer/LogVisualizer/Private/SVisualLoggerTab.h b/Engine/Source/Developer/LogVisualizer/Private/SVisualLoggerTab.h
new file mode 100644
--- /dev/null
+++ b/Engine/Source/Developer/LogVisualizer/Private/SVisualLoggerTab.h
@@ -0,0 +1,21 @@
+// Copyright Epic Games, Inc. All Rights Reserved.
+#pragma once
+
+/* Dependencies
+*****************************************************************************/
+
+#include "CoreMinimal.h"
+#include "Input/Reply.h"
+#include "Developer/LogVisualizer/Private/SVisualLogger.h"
+#include "Widgets/Docking/SDockTab.h"
+#include "LogVisualizerPublic.h"
+
+class SVisualLoggerTab : public SDockTab
+{
+public:
+	virtual bool SupportsKeyboardFocus() const override { return true; }
+	virtual FReply OnKeyDown(const FGeometry& MyGeometry, const FKeyEvent& InKeyEvent) override
+	{
+		return FLogVisualizer::Get().GetEvents().OnKeyboardEvent.Execute(MyGeometry, InKeyEvent);
+	}
+};
\ No newline at end of file
diff --git a/Engine/Source/Developer/LogVisualizer/Private/SVisualLoggerTimeline.cpp b/Engine/Source/Developer/LogVisualizer/Private/SVisualLoggerTimeline.cpp
--- a/Engine/Source/Developer/LogVisualizer/Private/SVisualLoggerTimeline.cpp
+++ b/Engine/Source/Developer/LogVisualizer/Private/SVisualLoggerTimeline.cpp
@@ -10,7 +10,6 @@
 #include "LogVisualizerSettings.h"
 #include "VisualLoggerDatabase.h"
 #include "LogVisualizerStyle.h"
-#include "LogVisualizerPrivate.h"
 #include "SVisualLoggerTimelinesContainer.h"
 #include "SVisualLoggerTimelineBar.h"
 
diff --git a/Engine/Source/Developer/LogVisualizer/Private/SVisualLoggerTimelineBar.cpp b/Engine/Source/Developer/LogVisualizer/Private/SVisualLoggerTimelineBar.cpp
--- a/Engine/Source/Developer/LogVisualizer/Private/SVisualLoggerTimelineBar.cpp
+++ b/Engine/Source/Developer/LogVisualizer/Private/SVisualLoggerTimelineBar.cpp
@@ -5,7 +5,7 @@
 #include "Rendering/DrawElements.h"
 #include "VisualLoggerDatabase.h"
 #include "LogVisualizerStyle.h"
-#include "LogVisualizerPrivate.h"
+#include "LogVisualizerPublic.h"
 #include "VisualLoggerTimeSliderController.h"
 #include "Misc/OutputDeviceHelper.h"
 
diff --git a/Engine/Source/Developer/LogVisualizer/Private/SVisualLoggerToolbar.h b/Engine/Source/Developer/LogVisualizer/Private/SVisualLoggerToolbar.h
--- a/Engine/Source/Developer/LogVisualizer/Private/SVisualLoggerToolbar.h
+++ b/Engine/Source/Developer/LogVisualizer/Private/SVisualLoggerToolbar.h
@@ -6,7 +6,7 @@
 #include "Widgets/SWidget.h"
 #include "Widgets/DeclarativeSyntaxSupport.h"
 #include "Framework/Commands/UICommandList.h"
-#include "Developer/LogVisualizer/Private/LogVisualizerPrivate.h"
+#include "SVisualLoggerBaseWidget.h"
 
 /**
 * Implements the messaging debugger toolbar widget.
diff --git a/Engine/Source/Developer/LogVisualizer/Private/SVisualLoggerView.h b/Engine/Source/Developer/LogVisualizer/Private/SVisualLoggerView.h
--- a/Engine/Source/Developer/LogVisualizer/Private/SVisualLoggerView.h
+++ b/Engine/Source/Developer/LogVisualizer/Private/SVisualLoggerView.h
@@ -8,9 +8,10 @@
 #include "Widgets/DeclarativeSyntaxSupport.h"
 #include "Framework/Commands/UICommandList.h"
 #include "Widgets/Layout/SSplitter.h"
-#include "Developer/LogVisualizer/Private/LogVisualizerPrivate.h"
 #include "Widgets/Input/SComboButton.h"
 #include "Widgets/Layout/SScrollBox.h"
+#include "SVisualLoggerBaseWidget.h"
+
 
 class SVisualLoggerTimelinesContainer;
 
diff --git a/Engine/Source/Developer/LogVisualizer/Private/VisualLoggerCanvasRenderer.cpp b/Engine/Source/Developer/LogVisualizer/Private/VisualLoggerCanvasRenderer.cpp
--- a/Engine/Source/Developer/LogVisualizer/Private/VisualLoggerCanvasRenderer.cpp
+++ b/Engine/Source/Developer/LogVisualizer/Private/VisualLoggerCanvasRenderer.cpp
@@ -9,7 +9,7 @@
 #include "LogVisualizerSettings.h"
 #include "LogVisualizerSessionSettings.h"
 #include "VisualLoggerDatabase.h"
-#include "LogVisualizerPrivate.h"
+#include "LogVisualizerPublic.h"
 #include "VisualLoggerTimeSliderController.h"
 #include "Debug/ReporterGraph.h"
 
diff --git a/Engine/Source/Developer/LogVisualizer/Private/VisualLoggerDatabase.cpp b/Engine/Source/Developer/LogVisualizer/Private/VisualLoggerDatabase.cpp
--- a/Engine/Source/Developer/LogVisualizer/Private/VisualLoggerDatabase.cpp
+++ b/Engine/Source/Developer/LogVisualizer/Private/VisualLoggerDatabase.cpp
@@ -5,7 +5,7 @@
 #include "Engine/World.h"
 #include "EngineUtils.h"
 #include "LogVisualizerSettings.h"
-#include "LogVisualizerPrivate.h"
+#include "LogVisualizerPublic.h"
 #include "VisualLoggerRenderingActor.h"
 
 TSharedPtr< struct FVisualLoggerDatabase > FVisualLoggerDatabase::StaticInstance = nullptr;
diff --git a/Engine/Source/Developer/LogVisualizer/Private/VisualLoggerRenderingActor.cpp b/Engine/Source/Developer/LogVisualizer/Private/VisualLoggerRenderingActor.cpp
--- a/Engine/Source/Developer/LogVisualizer/Private/VisualLoggerRenderingActor.cpp
+++ b/Engine/Source/Developer/LogVisualizer/Private/VisualLoggerRenderingActor.cpp
@@ -5,7 +5,7 @@
 #include "VisualLogger/VisualLogger.h"
 #include "LogVisualizerSettings.h"
 #include "VisualLoggerDatabase.h"
-#include "LogVisualizerPrivate.h"
+#include "LogVisualizerPublic.h"
 #if WITH_EDITOR
 #include "GeomTools.h"
 #endif // WITH_EDITOR
diff --git a/Engine/Source/Developer/LogVisualizer/Public/LogVisualizerPublic.h b/Engine/Source/Developer/LogVisualizer/Public/LogVisualizerPublic.h
new file mode 100644
--- /dev/null
+++ b/Engine/Source/Developer/LogVisualizer/Public/LogVisualizerPublic.h
@@ -0,0 +1,86 @@
+// Copyright Epic Games, Inc. All Rights Reserved.
+#pragma once
+
+/* Dependencies
+*****************************************************************************/
+#include "CoreMinimal.h"
+
+class FText;
+struct FGeometry;
+struct FKeyEvent;
+class FString;
+class UWorld;
+class UObject;
+
+DECLARE_DELEGATE_OneParam(FOnFiltersSearchChanged, const FText&);
+
+//DECLARE_DELEGATE(FOnFiltersChanged);
+DECLARE_MULTICAST_DELEGATE(FOnFiltersChanged);
+
+DECLARE_DELEGATE_ThreeParams(FOnLogLineSelectionChanged, TSharedPtr<struct FLogEntryItem> /*SelectedItem*/, int64 /*UserData*/, FName /*TagName*/);
+DECLARE_DELEGATE_RetVal_TwoParams(FReply, FOnKeyboardEvent, const FGeometry& /*MyGeometry*/, const FKeyEvent& /*InKeyEvent*/);
+DECLARE_DELEGATE_RetVal(float, FGetAnimationOutlinerFillPercentageFunc);
+
+struct FVisualLoggerEvents
+{
+	FOnFiltersChanged OnFiltersChanged;
+	FOnLogLineSelectionChanged OnLogLineSelectionChanged;
+	FOnKeyboardEvent OnKeyboardEvent;
+	FGetAnimationOutlinerFillPercentageFunc GetAnimationOutlinerFillPercentageFunc;
+};
+
+class FVisualLoggerTimeSliderController;
+struct LOGVISUALIZER_API FLogVisualizer
+{
+	/** LogVisualizer interface*/
+	void Reset();
+
+	FLinearColor GetColorForCategory(int32 Index) const;
+	FLinearColor GetColorForCategory(const FString& InFilterName) const;
+	UWorld* GetWorld(UObject* OptionalObject = nullptr);
+	FVisualLoggerEvents& GetEvents() { return VisualLoggerEvents; }
+
+	void SetCurrentVisualizer(TSharedPtr<class SVisualLogger> Visualizer) { CurrentVisualizer = Visualizer; }
+
+	void SetAnimationOutlinerFillPercentage(float FillPercentage) { AnimationOutlinerFillPercentage = FillPercentage; }
+	float GetAnimationOutlinerFillPercentage()
+	{
+		if (VisualLoggerEvents.GetAnimationOutlinerFillPercentageFunc.IsBound())
+		{
+			SetAnimationOutlinerFillPercentage(VisualLoggerEvents.GetAnimationOutlinerFillPercentageFunc.Execute());
+		}
+		return AnimationOutlinerFillPercentage;
+	}
+
+	int32 GetNextItem(FName RowName, int32 MoveDistance = 1);
+	int32 GetPreviousItem(FName RowName, int32 MoveDistance = 1);
+
+	void GotoNextItem(FName RowName, int32 MoveDistance = 1);
+	void GotoPreviousItem(FName RowName, int32 MoveDistance = 1);
+	void GotoFirstItem(FName RowName);
+	void GotoLastItem(FName RowName);
+
+	void UpdateCameraPosition(FName Rowname, int32 ItemIndes);
+
+	void SeekToTime(float Time);
+
+	/** Static access */
+	static void Initialize();
+	static void Shutdown();
+	static FLogVisualizer& Get();
+
+protected:
+	TSharedPtr<FVisualLoggerTimeSliderController> GetTimeSliderController() { return TimeSliderController; }
+
+	static TSharedPtr< struct FLogVisualizer > StaticInstance;
+
+	TSharedPtr<FVisualLoggerTimeSliderController> TimeSliderController;
+	FVisualLoggerEvents VisualLoggerEvents;
+	TWeakPtr<class SVisualLogger> CurrentVisualizer;
+	float AnimationOutlinerFillPercentage;
+
+	friend class SVisualLoggerViewer;
+	friend class SVisualLoggerView;
+	friend class SVisualLogger;
+	friend struct FVisualLoggerCanvasRenderer;
+};
\ No newline at end of file
