Added missing include for non unity build from previously integrated PR

diff --git a/Engine/Source/Developer/LogVisualizer/Private/SVisualLoggerStatusView.h b/Engine/Source/Developer/LogVisualizer/Private/SVisualLoggerStatusView.h
--- a/Engine/Source/Developer/LogVisualizer/Private/SVisualLoggerStatusView.h
+++ b/Engine/Source/Developer/LogVisualizer/Private/SVisualLoggerStatusView.h
@@ -9,6 +9,7 @@
 #include "Widgets/Views/STableViewBase.h"
 #include "Widgets/Views/STableRow.h"
 #include "Widgets/Views/STreeView.h"
+#include "SVisualLoggerBaseWidget.h"
 
 struct FLogStatusItem;
 struct FVisualLoggerDBRow;
