Made sure no one will accidentally place VisualLoggerRenderingActor on the level.

diff --git a/Engine/Source/Developer/LogVisualizer/Private/VisualLoggerRenderingActor.h b/Engine/Source/Developer/LogVisualizer/Private/VisualLoggerRenderingActor.h
--- a/Engine/Source/Developer/LogVisualizer/Private/VisualLoggerRenderingActor.h
+++ b/Engine/Source/Developer/LogVisualizer/Private/VisualLoggerRenderingActor.h
@@ -18,7 +18,7 @@ struct FVisualLoggerDBRow;
 
 DECLARE_MULTICAST_DELEGATE_OneParam(FOnSelectionChanged, class AActor*);
 
-UCLASS(config = Engine, NotBlueprintable, Transient)
+UCLASS(config = Engine, NotBlueprintable, Transient, notplaceable, AdvancedClassDisplay)
 class LOGVISUALIZER_API AVisualLoggerRenderingActor : public AActor
 {
 	struct FTimelineDebugShapes
