[LogVisualizer] Fixed some typos

diff --git a/Engine/Source/Developer/LogVisualizer/Private/LogVisualizerSettings.cpp b/Engine/Source/Developer/LogVisualizer/Private/LogVisualizerSettings.cpp
--- a/Engine/Source/Developer/LogVisualizer/Private/LogVisualizerSettings.cpp
+++ b/Engine/Source/Developer/LogVisualizer/Private/LogVisualizerSettings.cpp
@@ -33,45 +33,45 @@ class UMaterial* ULogVisualizerSettings::GetDebugMeshMaterial()
 	return DebugMeshMaterialFakeLight;
 }
 
-void ULogVisualizerSettings::SavePresistentData()
+void ULogVisualizerSettings::SavePersistentData()
 {
-	if (bPresistentFilters)
+	if (bPersistentFilters)
 	{
-		PresistentFilters = FVisualLoggerFilters::Get();
-		for (int32 Index = PresistentFilters.Categories.Num() - 1; Index >= 0; --Index)
+		PersistentFilters = FVisualLoggerFilters::Get();
+		for (int32 Index = PersistentFilters.Categories.Num() - 1; Index >= 0; --Index)
 		{
-			FCategoryFilter& Category = PresistentFilters.Categories[Index];
+			FCategoryFilter& Category = PersistentFilters.Categories[Index];
 			if (Category.bIsInUse == false)
 			{
-				PresistentFilters.Categories.RemoveAt(Index);
+				PersistentFilters.Categories.RemoveAt(Index);
 			}
 		}
 	}
 	else
 	{
-		PresistentFilters = FVisualLoggerFilters();
+		PersistentFilters = FVisualLoggerFilters();
 	}
 	SaveConfig();
 }
 
-void ULogVisualizerSettings::ClearPresistentData()
+void ULogVisualizerSettings::ClearPersistentData()
 {
-	if (bPresistentFilters)
+	if (bPersistentFilters)
 	{
-		PresistentFilters = FVisualLoggerFilters();
+		PersistentFilters = FVisualLoggerFilters();
 	}
 }
 
-void ULogVisualizerSettings::LoadPresistentData()
+void ULogVisualizerSettings::LoadPersistentData()
 {
-	if (bPresistentFilters)
+	if (bPersistentFilters)
 	{
-		for (int32 Index = PresistentFilters.Categories.Num() - 1; Index >= 0; --Index)
+		for (int32 Index = PersistentFilters.Categories.Num() - 1; Index >= 0; --Index)
 		{
-			FCategoryFilter& Category = PresistentFilters.Categories[Index];
+			FCategoryFilter& Category = PersistentFilters.Categories[Index];
 			Category.bIsInUse = false;
 		}
-		FVisualLoggerFilters::Get().InitWith(PresistentFilters);
+		FVisualLoggerFilters::Get().InitWith(PersistentFilters);
 	}
 	else
 	{
diff --git a/Engine/Source/Developer/LogVisualizer/Private/SVisualLogger.cpp b/Engine/Source/Developer/LogVisualizer/Private/SVisualLogger.cpp
--- a/Engine/Source/Developer/LogVisualizer/Private/SVisualLogger.cpp
+++ b/Engine/Source/Developer/LogVisualizer/Private/SVisualLogger.cpp
@@ -116,7 +116,7 @@ SVisualLogger::~SVisualLogger()
 
 	UWorld* World = NULL;
 #if WITH_EDITOR
-	ULogVisualizerSettings::StaticClass()->GetDefaultObject<ULogVisualizerSettings>()->SavePresistentData();
+	ULogVisualizerSettings::StaticClass()->GetDefaultObject<ULogVisualizerSettings>()->SavePersistentData();
 #endif
 
 	if (LastUsedWorld.IsValid())
@@ -171,7 +171,7 @@ void SVisualLogger::Construct(const FArguments& InArgs, const TSharedRef<SDockTa
 
 	ULogVisualizerSettings* Settings = ULogVisualizerSettings::StaticClass()->GetDefaultObject<ULogVisualizerSettings>();
 	Settings->ConfigureVisLog();
-	Settings->LoadPresistentData();
+	Settings->LoadPersistentData();
 
 	ActionList.MapAction(Commands.StartRecording, FExecuteAction::CreateRaw(this, &SVisualLogger::HandleStartRecordingCommandExecute), FCanExecuteAction::CreateRaw(this, &SVisualLogger::HandleStartRecordingCommandCanExecute), FIsActionChecked(), FIsActionButtonVisible::CreateRaw(this, &SVisualLogger::HandleStartRecordingCommandIsVisible));
 	ActionList.MapAction(Commands.StopRecording, FExecuteAction::CreateRaw(this, &SVisualLogger::HandleStopRecordingCommandExecute), FCanExecuteAction::CreateRaw(this, &SVisualLogger::HandleStopRecordingCommandCanExecute), FIsActionChecked(), FIsActionButtonVisible::CreateRaw(this, &SVisualLogger::HandleStopRecordingCommandIsVisible));
diff --git a/Engine/Source/Developer/LogVisualizer/Public/LogVisualizerSettings.h b/Engine/Source/Developer/LogVisualizer/Public/LogVisualizerSettings.h
--- a/Engine/Source/Developer/LogVisualizer/Public/LogVisualizerSettings.h
+++ b/Engine/Source/Developer/LogVisualizer/Public/LogVisualizerSettings.h
@@ -147,7 +147,7 @@ public:
 
 	/**Whether to store all filter settings on exit*/
 	UPROPERTY(EditAnywhere, config, Category = "VisualLogger")
-	bool bPresistentFilters;
+	bool bPersistentFilters;
 
 	/** Whether to extreme values on graph (data has to be provided for extreme values) */
 	UPROPERTY(EditAnywhere, config, Category = "VisualLogger")
@@ -177,17 +177,17 @@ public:
 
 	class UMaterial* GetDebugMeshMaterial();
 
-	void SavePresistentData();
+	void SavePersistentData();
 
-	void ClearPresistentData();
+	void ClearPersistentData();
 
-	void LoadPresistentData();
+	void LoadPersistentData();
 
 	void ConfigureVisLog();
 
 protected:
 	UPROPERTY(config)
-	FVisualLoggerFiltersData PresistentFilters;
+	FVisualLoggerFiltersData PersistentFilters;
 
 	/** A material used to render debug meshes with kind of flat shading, mostly used by Visual Logger tool. */
 	UPROPERTY()
