Removed reference to ue4

diff --git a/Engine/Source/Developer/LogVisualizer/Private/LogVisualizerModule.cpp b/Engine/Source/Developer/LogVisualizer/Private/LogVisualizerModule.cpp
--- a/Engine/Source/Developer/LogVisualizer/Private/LogVisualizerModule.cpp
+++ b/Engine/Source/Developer/LogVisualizer/Private/LogVisualizerModule.cpp
@@ -63,7 +63,7 @@ void FLogVisualizerModule::StartupModule()
 	{
 		SettingsModule->RegisterSettings("Editor", "Advanced", "VisualLogger",
 			LOCTEXT("AIToolsSettingsName", "Visual Logger"),
-			LOCTEXT("AIToolsSettingsDescription", "General settings for UE4 AI Tools."),
+			LOCTEXT("AIToolsSettingsDescription", "General settings for UE AI Tools."),
 			ULogVisualizerSettings::StaticClass()->GetDefaultObject()
 			);
 	}
