[VisualLogger] Renamed some slate widgets to fix some symbol mixup between modules

diff --git a/Engine/Source/Developer/LogVisualizer/Private/SFilterCheckBox.h b/Engine/Source/Developer/LogVisualizer/Private/SFilterCheckBox.h
deleted file mode 100644
diff --git a/Engine/Source/Developer/LogVisualizer/Private/SFilterWidget.cpp b/Engine/Source/Developer/LogVisualizer/Private/SFilterWidget.cpp
deleted file mode 100644
diff --git a/Engine/Source/Developer/LogVisualizer/Private/SFilterWidget.h b/Engine/Source/Developer/LogVisualizer/Private/SFilterWidget.h
deleted file mode 100644
diff --git a/Engine/Source/Developer/LogVisualizer/Private/SVisualLoggerFilterCheckBox.h b/Engine/Source/Developer/LogVisualizer/Private/SVisualLoggerFilterCheckBox.h
new file mode 100644
--- /dev/null
+++ b/Engine/Source/Developer/LogVisualizer/Private/SVisualLoggerFilterCheckBox.h
@@ -0,0 +1,50 @@
+// Copyright Epic Games, Inc. All Rights Reserved.
+#pragma once
+
+#include "CoreMinimal.h"
+#include "InputCoreTypes.h"
+#include "Input/Reply.h"
+#include "Widgets/Input/SCheckBox.h"
+
+/** A class for check boxes in the filter list. If you double click a filter check box, you will enable it and disable all others */
+class SVisualLoggerFilterCheckBox : public SCheckBox
+{
+public:
+	void SetOnFilterDoubleClicked(const FOnClicked& NewFilterDoubleClicked)
+	{
+		OnFilterDoubleClicked = NewFilterDoubleClicked;
+	}
+
+	void SetOnFilterMiddleButtonClicked(const FOnClicked& NewFilterMiddleButtonClicked)
+	{
+		OnFilterMiddleButtonClicked = NewFilterMiddleButtonClicked;
+	}
+
+	virtual FReply OnMouseButtonDoubleClick(const FGeometry& InMyGeometry, const FPointerEvent& InMouseEvent) override
+	{
+		if (InMouseEvent.GetEffectingButton() == EKeys::LeftMouseButton && OnFilterDoubleClicked.IsBound())
+		{
+			return OnFilterDoubleClicked.Execute();
+		}
+		else
+		{
+			return SCheckBox::OnMouseButtonDoubleClick(InMyGeometry, InMouseEvent);
+		}
+	}
+
+	virtual FReply OnMouseButtonUp(const FGeometry& InMyGeometry, const FPointerEvent& InMouseEvent) override
+	{
+		if (InMouseEvent.GetEffectingButton() == EKeys::MiddleMouseButton && OnFilterMiddleButtonClicked.IsBound())
+		{
+			return OnFilterMiddleButtonClicked.Execute();
+		}
+		else
+		{
+			return SCheckBox::OnMouseButtonUp(InMyGeometry, InMouseEvent);
+		}
+	}
+
+private:
+	FOnClicked OnFilterDoubleClicked;
+	FOnClicked OnFilterMiddleButtonClicked;
+};
diff --git a/Engine/Source/Developer/LogVisualizer/Private/SVisualLoggerFilterWidget.cpp b/Engine/Source/Developer/LogVisualizer/Private/SVisualLoggerFilterWidget.cpp
new file mode 100644
--- /dev/null
+++ b/Engine/Source/Developer/LogVisualizer/Private/SVisualLoggerFilterWidget.cpp
@@ -0,0 +1,185 @@
+// Copyright Epic Games, Inc. All Rights Reserved.
+
+#include "SVisualLoggerFilterWidget.h"
+#include "Misc/OutputDeviceHelper.h"
+#include "Textures/SlateIcon.h"
+#include "Framework/Commands/UIAction.h"
+#include "Widgets/Layout/SBorder.h"
+#include "Widgets/Text/STextBlock.h"
+#include "Framework/MultiBox/MultiBoxBuilder.h"
+#include "LogVisualizerSettings.h"
+#include "LogVisualizerStyle.h"
+
+#define LOCTEXT_NAMESPACE "SVisualLoggerFilterWidget"
+/** Constructs this widget with InArgs */
+void SVisualLoggerFilterWidget::Construct(const FArguments& InArgs)
+{
+	OnFilterChanged = InArgs._OnFilterChanged;
+	OnRequestRemove = InArgs._OnRequestRemove;
+	OnRequestEnableOnly = InArgs._OnRequestEnableOnly;
+	OnRequestDisableAll = InArgs._OnRequestDisableAll;
+	OnRequestRemoveAll = InArgs._OnRequestRemoveAll;
+	FilterName = InArgs._FilterName;
+	ColorCategory = InArgs._ColorCategory;
+	BorderBackgroundColor = FLinearColor(0.2f, 0.2f, 0.2f, 0.2f);
+
+	LastVerbosity = ELogVerbosity::NoLogging;
+	GetCaptionString(); // cache caption string
+
+	bWasEnabledLastTime = !IsEnabled();
+	GetTooltipString(); // cache tooltip string
+
+	// Get the tooltip and color of the type represented by this filter
+	FilterColor = ColorCategory;
+
+	ChildSlot
+		[
+			SNew(SBorder)
+			.Padding(2)
+			.BorderBackgroundColor(this, &SVisualLoggerFilterWidget::GetBorderBackgroundColor)
+			.BorderImage(FLogVisualizerStyle::Get().GetBrush("ContentBrowser.FilterButtonBorder"))
+			[
+				SAssignNew(ToggleButtonPtr, SVisualLoggerFilterCheckBox)
+				.Style(FLogVisualizerStyle::Get(), "ContentBrowser.FilterButton")
+				.ToolTipText(this, &SVisualLoggerFilterWidget::GetTooltipString)
+				.Padding(this, &SVisualLoggerFilterWidget::GetFilterNamePadding)
+				.IsChecked(this, &SVisualLoggerFilterWidget::IsChecked)
+				.OnCheckStateChanged(this, &SVisualLoggerFilterWidget::FilterToggled)
+				.OnGetMenuContent(this, &SVisualLoggerFilterWidget::GetRightClickMenuContent)
+				.ForegroundColor(this, &SVisualLoggerFilterWidget::GetFilterForegroundColor)
+				[
+					SNew(STextBlock)
+					.ColorAndOpacity(this, &SVisualLoggerFilterWidget::GetFilterNameColorAndOpacity)
+					.Text(this, &SVisualLoggerFilterWidget::GetCaptionString)
+				]
+			]
+		];
+
+	ToggleButtonPtr->SetOnFilterDoubleClicked(FOnClicked::CreateSP(this, &SVisualLoggerFilterWidget::FilterDoubleClicked));
+	ToggleButtonPtr->SetOnFilterMiddleButtonClicked(FOnClicked::CreateSP(this, &SVisualLoggerFilterWidget::FilterMiddleButtonClicked));
+}
+
+FText SVisualLoggerFilterWidget::GetCaptionString() const
+{
+	FString CaptionString;
+	FCategoryFilter& CategoryFilter = FVisualLoggerFilters::Get().GetCategoryByName(GetFilterNameAsString());
+	if (CategoryFilter.LogVerbosity != LastVerbosity)
+	{
+		const FString VerbosityStr = FOutputDeviceHelper::VerbosityToString((ELogVerbosity::Type)CategoryFilter.LogVerbosity);
+		if (CategoryFilter.LogVerbosity != ELogVerbosity::VeryVerbose)
+		{
+			CaptionString = FString::Printf(TEXT("%s [%s]"), *GetFilterNameAsString().Replace(TEXT("Log"), TEXT(""), ESearchCase::CaseSensitive), *VerbosityStr.Mid(0, 1));
+		}
+		else
+		{
+			CaptionString = FString::Printf(TEXT("%s [VV]"), *GetFilterNameAsString().Replace(TEXT("Log"), TEXT(""), ESearchCase::CaseSensitive));
+		}
+
+		CachedCaptionString = FText::FromString(CaptionString);
+		LastVerbosity = (ELogVerbosity::Type)CategoryFilter.LogVerbosity;
+	}
+	return CachedCaptionString;
+}
+
+FText SVisualLoggerFilterWidget::GetTooltipString() const
+{
+	if (bWasEnabledLastTime != IsEnabled())
+	{
+		FCategoryFilter& CategoryFilter = FVisualLoggerFilters::Get().GetCategoryByName(GetFilterNameAsString());
+		const FString VerbosityStr = FOutputDeviceHelper::VerbosityToString((ELogVerbosity::Type)CategoryFilter.LogVerbosity);
+
+		CachedTooltipString = FText::FromString(
+			IsEnabled() ?
+			FString::Printf(TEXT("Enabled '%s' category for '%s' verbosity and lower\nRight click to change verbosity"), *GetFilterNameAsString(), *VerbosityStr)
+			:
+			FString::Printf(TEXT("Disabled '%s' category"), *GetFilterNameAsString())
+			);
+
+		bWasEnabledLastTime = IsEnabled();
+	}
+
+	return CachedTooltipString;
+}
+
+TSharedRef<SWidget> SVisualLoggerFilterWidget::GetRightClickMenuContent()
+{
+	FMenuBuilder MenuBuilder(true, NULL);
+
+	if (IsEnabled())
+	{
+		FText FiletNameAsText = FText::FromString(GetFilterNameAsString());
+		MenuBuilder.BeginSection("VerbositySelection", LOCTEXT("VerbositySelection", "Current verbosity selection"));
+		{
+			for (int32 Index = ELogVerbosity::NoLogging + 1; Index <= ELogVerbosity::VeryVerbose; Index++)
+			{
+				const FString VerbosityStr = FOutputDeviceHelper::VerbosityToString((ELogVerbosity::Type)Index);
+				MenuBuilder.AddMenuEntry(
+					FText::Format(LOCTEXT("UseVerbosity", "Use: {0}"), FText::FromString(VerbosityStr)),
+					LOCTEXT("UseVerbosityTooltip", "Applay verbosity to selected flter."),
+					FSlateIcon(),
+					FUIAction(FExecuteAction::CreateSP(this, &SVisualLoggerFilterWidget::SetVerbosityFilter, Index))
+					);
+			}
+		}
+		MenuBuilder.EndSection();
+	}
+	MenuBuilder.BeginSection("FilterAction", LOCTEXT("FilterAction", "Context actions"));
+		MenuBuilder.AddMenuEntry(
+			LOCTEXT("DisableAllButThis", "Disable all but this"),
+			LOCTEXT("HideAllButThisTooltip", "Disable all other categories"),
+			FSlateIcon(),
+			FUIAction(FExecuteAction::CreateSP(this, &SVisualLoggerFilterWidget::DisableAllButThis))
+			);
+		MenuBuilder.AddMenuEntry(
+			LOCTEXT("EnableAll", "Enable all categories"),
+			LOCTEXT("EnableAllTooltip", "Enable all categories"),
+			FSlateIcon(),
+			FUIAction(FExecuteAction::CreateSP(this, &SVisualLoggerFilterWidget::EnableAllCategories))
+			);
+	MenuBuilder.EndSection();
+
+	return MenuBuilder.MakeWidget();
+}
+
+bool SVisualLoggerFilterWidget::IsEnabled() const 
+{ 
+	FCategoryFilter& CategoryFilter = FVisualLoggerFilters::Get().GetCategoryByName(GetFilterNameAsString());
+	return CategoryFilter.Enabled;
+}
+
+void SVisualLoggerFilterWidget::SetEnabled(bool InEnabled)
+{
+	FCategoryFilter& CategoryFilter = FVisualLoggerFilters::Get().GetCategoryByName(GetFilterNameAsString());
+	if (InEnabled != CategoryFilter.Enabled)
+	{
+		CategoryFilter.Enabled = InEnabled;
+		OnFilterChanged.ExecuteIfBound();
+	}
+}
+
+void SVisualLoggerFilterWidget::FilterToggled(ECheckBoxState NewState)
+{
+	SetEnabled(NewState == ECheckBoxState::Checked);
+}
+
+void SVisualLoggerFilterWidget::SetVerbosityFilter(int32 SelectedVerbosityIndex)
+{
+	FCategoryFilter& CategoryFilter = FVisualLoggerFilters::Get().GetCategoryByName(GetFilterNameAsString());
+	CategoryFilter.LogVerbosity = (ELogVerbosity::Type)SelectedVerbosityIndex;
+	OnFilterChanged.ExecuteIfBound();
+}
+
+void SVisualLoggerFilterWidget::DisableAllButThis()
+{
+	FVisualLoggerFilters::Get().DeactivateAllButThis(GetFilterNameAsString());
+	OnFilterChanged.ExecuteIfBound();
+}
+
+void SVisualLoggerFilterWidget::EnableAllCategories()
+{
+	FVisualLoggerFilters::Get().EnableAllCategories();
+	OnFilterChanged.ExecuteIfBound();
+}
+
+
+#undef LOCTEXT_NAMESPACE
diff --git a/Engine/Source/Developer/LogVisualizer/Private/SVisualLoggerFilterWidget.h b/Engine/Source/Developer/LogVisualizer/Private/SVisualLoggerFilterWidget.h
new file mode 100644
--- /dev/null
+++ b/Engine/Source/Developer/LogVisualizer/Private/SVisualLoggerFilterWidget.h
@@ -0,0 +1,176 @@
+// Copyright Epic Games, Inc. All Rights Reserved.
+#pragma once
+
+#include "CoreMinimal.h"
+#include "Styling/SlateColor.h"
+#include "Input/Reply.h"
+#include "Widgets/SWidget.h"
+#include "Widgets/DeclarativeSyntaxSupport.h"
+#include "Layout/Margin.h"
+#include "Styling/SlateTypes.h"
+#include "Widgets/SCompoundWidget.h"
+#include "SVisualLoggerFilterCheckBox.h"
+
+/**
+* A single filter in the filter list. Can be removed by clicking the remove button on it.
+*/
+class SVisualLoggerFilterWidget : public SCompoundWidget
+{
+public:
+	DECLARE_DELEGATE_OneParam(FOnFilterRequest, const TSharedRef<SVisualLoggerFilterWidget>&);
+	DECLARE_DELEGATE(FOnSimpleRequest);
+
+	SLATE_BEGIN_ARGS(SVisualLoggerFilterWidget){}
+		/** If this is an front end filter, this is the filter object */
+	SLATE_ARGUMENT(FName, FilterName)
+		/** Color selected for filter*/
+		SLATE_ARGUMENT(FLinearColor, ColorCategory)
+		/** Invoked when the filter toggled */
+		SLATE_EVENT(FOnSimpleRequest, OnFilterChanged)
+		/** Invoked when a request to remove this filter originated from within this filter */
+		SLATE_EVENT(FOnFilterRequest, OnRequestRemove)
+		/** Invoked when a request to enable only this filter originated from within this filter */
+		SLATE_EVENT(FOnFilterRequest, OnRequestEnableOnly)
+		/** Invoked when a request to disable all filters originated from within this filter */
+		SLATE_EVENT(FOnSimpleRequest, OnRequestDisableAll)
+		/** Invoked when a request to remove all filters originated from within this filter */
+		SLATE_EVENT(FOnSimpleRequest, OnRequestRemoveAll)
+	SLATE_END_ARGS()
+
+	/** Constructs this widget with InArgs */
+	void Construct(const FArguments& InArgs);
+
+	/** Caption String getter */
+	FText GetCaptionString() const;
+	/** Tool Tip string getter */
+	FText GetTooltipString() const;
+
+	/** Sets whether or not this filter is applied to the combined filter */
+	void SetEnabled(bool InEnabled);
+
+	/** Returns true if this filter contributes to the combined filter */
+	bool IsEnabled() const;
+
+	FORCEINLINE FName GetFilterName() const { return FilterName; }
+
+	/** Returns the display name for this filter */
+	FORCEINLINE FString GetFilterNameAsString() const
+	{
+		static const FString UnknownFilter(TEXT("UnknownFilter"));
+		if (FilterName != NAME_None)
+		{
+			return FilterName.ToString();
+		}
+
+		return UnknownFilter;
+	}
+
+	void SetBorderBackgroundColor(FLinearColor InColor) { BorderBackgroundColor = InColor; }
+
+private:
+	/** Handler for when the filter checkbox is clicked */
+	void FilterToggled(ECheckBoxState NewState);
+
+	/** Handler for when the filter checkbox is double clicked */
+	FReply FilterDoubleClicked()
+	{
+		// Disable all other filters and enable this one.
+		OnRequestDisableAll.ExecuteIfBound();
+		SetEnabled(true);
+
+		return FReply::Handled();
+	}
+
+	/** Handler for when the filter checkbox is middle button clicked */
+	FReply FilterMiddleButtonClicked()
+	{
+		RemoveFilter();
+		return FReply::Handled();
+	}
+
+	/** Handler to create a right click menu */
+	TSharedRef<SWidget> GetRightClickMenuContent();
+
+	/** Removes this filter from the filter list */
+	void SetVerbosityFilter(int32 SelectedVerbosityIndex);
+
+	void DisableAllButThis();
+
+	void EnableAllCategories();
+
+	void RemoveFilter()
+	{
+		TSharedRef<SVisualLoggerFilterWidget> Self = SharedThis(this);
+		OnRequestRemove.ExecuteIfBound(Self);
+	}
+
+	/** Enables only this filter from the filter list */
+	void EnableOnly()
+	{
+		TSharedRef<SVisualLoggerFilterWidget> Self = SharedThis(this);
+		OnRequestEnableOnly.ExecuteIfBound(Self);
+	}
+
+	/** Disables all active filters in the list */
+	void DisableAllFilters() { OnRequestDisableAll.ExecuteIfBound(); }
+
+	/** Removes all filters in the list */
+	void RemoveAllFilters() { OnRequestRemoveAll.ExecuteIfBound(); }
+
+	/** Handler to determine the "checked" state of the filter checkbox */
+	ECheckBoxState IsChecked() const { return IsEnabled() ? ECheckBoxState::Checked : ECheckBoxState::Unchecked; }
+
+	/** Handler to determine the color of the checkbox when it is checked */
+	FSlateColor GetFilterForegroundColor() const { return IsChecked() == ECheckBoxState::Checked ? FilterColor : FLinearColor::White; }
+
+	/** Handler to determine the padding of the checkbox text when it is pressed */
+	FMargin GetFilterNamePadding() const { return ToggleButtonPtr->IsPressed() ? FMargin(3, 2, 4, 0) : FMargin(3, 1, 4, 1); }
+
+	/** Handler to determine the color of the checkbox text when it is hovered */
+	FSlateColor GetFilterNameColorAndOpacity() const
+	{
+		const float DimFactor = 0.75f;
+		return IsChecked() == ECheckBoxState::Checked ? (IsHovered() ? ColorCategory * DimFactor : ColorCategory) : (IsHovered() ? FLinearColor::White : FLinearColor::White * DimFactor);
+	}
+
+	FSlateColor GetBorderBackgroundColor() const
+	{
+		return BorderBackgroundColor;
+	}
+
+private:
+	/** Invoked when the filter toggled */
+	FOnSimpleRequest OnFilterChanged;
+
+	/** Invoked when a request to remove this filter originated from within this filter */
+	FOnFilterRequest OnRequestRemove;
+
+	/** Invoked when a request to enable only this filter originated from within this filter */
+	FOnFilterRequest OnRequestEnableOnly;
+
+	/** Invoked when a request to disable all filters originated from within this filter */
+	FOnSimpleRequest OnRequestDisableAll;
+
+	/** Invoked when a request to remove all filters originated from within this filter */
+	FOnSimpleRequest OnRequestRemoveAll;
+
+	FName FilterName;
+
+	FLinearColor ColorCategory;
+
+	/** The button to toggle the filter on or off */
+	TSharedPtr<SVisualLoggerFilterCheckBox> ToggleButtonPtr;
+
+	/** The color of the checkbox for this filter */
+	FLinearColor FilterColor;
+
+	/** Default color for border background */
+	FLinearColor BorderBackgroundColor;
+
+	/** hacky to cache data */
+	mutable bool bWasEnabledLastTime;
+	mutable ELogVerbosity::Type LastVerbosity;
+
+	mutable FText CachedCaptionString;
+	mutable FText CachedTooltipString;
+};
diff --git a/Engine/Source/Developer/LogVisualizer/Private/SVisualLoggerFilters.cpp b/Engine/Source/Developer/LogVisualizer/Private/SVisualLoggerFilters.cpp
--- a/Engine/Source/Developer/LogVisualizer/Private/SVisualLoggerFilters.cpp
+++ b/Engine/Source/Developer/LogVisualizer/Private/SVisualLoggerFilters.cpp
@@ -11,7 +11,7 @@
 #include "LogVisualizerSettings.h"
 #include "VisualLoggerDatabase.h"
 #include "LogVisualizerStyle.h"
-#include "SFilterWidget.h"
+#include "SVisualLoggerFilterWidget.h"
 #include "Widgets/Input/SSearchBox.h"
 #include "LogVisualizerPublic.h"
 
@@ -327,7 +327,7 @@ void SVisualLoggerFilters::AddFilterCategory(FString InName, ELogVerbosity::Type
 	}
 	else
 	{
-		for (TSharedRef<SFilterWidget>& CurrentFilter : Filters)
+		for (TSharedRef<SVisualLoggerFilterWidget>& CurrentFilter : Filters)
 		{
 			if (CurrentFilter->GetFilterNameAsString() == InName)
 			{
@@ -336,11 +336,11 @@ void SVisualLoggerFilters::AddFilterCategory(FString InName, ELogVerbosity::Type
 		}
 
 		const FLinearColor Color = FLogVisualizer::Get().GetColorForCategory(InName);
-		TSharedRef<SFilterWidget> NewFilter =
-			SNew(SFilterWidget)
+		TSharedRef<SVisualLoggerFilterWidget> NewFilter =
+			SNew(SVisualLoggerFilterWidget)
 			.FilterName(*InName)
 			.ColorCategory(Color)
-			.OnFilterChanged(SFilterWidget::FOnSimpleRequest::CreateRaw(this, &SVisualLoggerFilters::OnFiltersChanged));
+			.OnFilterChanged(SVisualLoggerFilterWidget::FOnSimpleRequest::CreateRaw(this, &SVisualLoggerFilters::OnFiltersChanged));
 
 		Filters.Add(NewFilter);
 		FilterBox->AddSlot()
@@ -380,7 +380,7 @@ void SVisualLoggerFilters::OnFilterCategoryRemoved(FString InName)
 	}
 	else
 	{
-		for (TSharedRef<SFilterWidget> CurrentFilter : Filters)
+		for (TSharedRef<SVisualLoggerFilterWidget> CurrentFilter : Filters)
 		{
 			if (CurrentFilter->GetFilterNameAsString() == InName)
 			{
@@ -395,7 +395,7 @@ void SVisualLoggerFilters::OnFilterCategoryRemoved(FString InName)
 void SVisualLoggerFilters::OnFiltersChanged()
 {
 	TArray<FString> EnabledFilters;
-	for (TSharedRef<SFilterWidget> CurrentFilter : Filters)
+	for (TSharedRef<SVisualLoggerFilterWidget> CurrentFilter : Filters)
 	{
 		if (CurrentFilter->IsEnabled())
 		{
@@ -426,7 +426,7 @@ void SVisualLoggerFilters::OnItemsSelectionChanged(const FVisualLoggerDBRow& Cha
 
 	for (int32 Index = 0; Index < Filters.Num(); ++Index)
 	{
-		SFilterWidget& Filter = Filters[Index].Get();
+		SVisualLoggerFilterWidget& Filter = Filters[Index].Get();
 		Filter.SetBorderBackgroundColor(FLinearColor(0.2f, 0.2f, 0.2f, 0.2f));
 		for (const FVisualLoggerCategoryVerbosityPair& Category : Categories)
 		{
diff --git a/Engine/Source/Developer/LogVisualizer/Private/SVisualLoggerFilters.h b/Engine/Source/Developer/LogVisualizer/Private/SVisualLoggerFilters.h
--- a/Engine/Source/Developer/LogVisualizer/Private/SVisualLoggerFilters.h
+++ b/Engine/Source/Developer/LogVisualizer/Private/SVisualLoggerFilters.h
@@ -10,14 +10,13 @@
 
 class FMenuBuilder;
 class SComboButton;
-class SFilterWidget;
+class SVisualLoggerFilterWidget;
 class SWrapBox;
 struct FVisualLoggerDBRow;
 
 /**
 * A list of filters currently applied to an asset view.
 */
-class SFilterWidget;
 class SVisualLoggerFilters : public SVisualLoggerBaseWidget
 {
 public:
@@ -58,7 +57,7 @@ protected:
 protected:
 	/** The horizontal box which contains all the filters */
 	TSharedPtr<SWrapBox> FilterBox;
-	TArray<TSharedRef<SFilterWidget> > Filters;
+	TArray<TSharedRef<SVisualLoggerFilterWidget>> Filters;
 	TSharedPtr<SComboButton> GraphsFilterCombo;
 
 	FString GraphsSearchString;
