[VisualLogger] Reset data on PIE (also clear all entries if bResetDataWithNewSession)

diff --git a/Engine/Source/Developer/LogVisualizer/Private/SVisualLogger.cpp b/Engine/Source/Developer/LogVisualizer/Private/SVisualLogger.cpp
--- a/Engine/Source/Developer/LogVisualizer/Private/SVisualLogger.cpp
+++ b/Engine/Source/Developer/LogVisualizer/Private/SVisualLogger.cpp
@@ -7,7 +7,6 @@
 #include "Modules/ModuleManager.h"
 #include "Async/Future.h"
 #include "Async/Async.h"
-#include "EngineGlobals.h"
 #include "Debug/DebugDrawService.h"
 #include "AI/NavigationSystemBase.h"
 #include "Engine/Engine.h"
@@ -109,12 +108,15 @@ SVisualLogger::SVisualLogger()
 
 SVisualLogger::~SVisualLogger()
 {
+#if WITH_EDITOR
+	FEditorDelegates::PostPIEStarted.Remove(PostPIEStartedHandle);
+#endif
+
 	GEngine->OnWorldAdded().RemoveAll(this);
 
 	FVisualLogger::Get().RemoveDevice(InternalDevice.Get());
 	InternalDevice.Reset();
 
-	UWorld* World = NULL;
 #if WITH_EDITOR
 	ULogVisualizerSettings::StaticClass()->GetDefaultObject<ULogVisualizerSettings>()->SavePersistentData();
 #endif
@@ -164,6 +166,17 @@ void SVisualLogger::Construct(const FArguments& InArgs, const TSharedRef<SDockTa
 
 	GEngine->OnWorldAdded().AddRaw(this, &SVisualLogger::OnNewWorld);
 
+#if WITH_EDITOR
+	PostPIEStartedHandle = FEditorDelegates::PostPIEStarted.AddLambda([this](const bool bIsSimulating)
+	{
+		UWorld* World = FLogVisualizer::Get().GetWorld();
+		if (World != nullptr && World != LastUsedWorld)
+		{
+			OnNewWorld(World);
+		}
+	});
+#endif
+
 	//////////////////////////////////////////////////////////////////////////
 	// Command Action Lists
 	const FVisualLoggerCommands& Commands = FVisualLoggerCommands::Get();
@@ -777,8 +790,6 @@ void SVisualLogger::OnItemsSelectionChanged(const FVisualLoggerDBRow& ChangedRow
 void SVisualLogger::OnFiltersChanged()
 {
 	const uint32 StartCycles = FPlatformTime::Cycles();
-
-	ULogVisualizerSettings* Settings = ULogVisualizerSettings::StaticClass()->GetDefaultObject<ULogVisualizerSettings>();
 	const FString QuickSearchStrng = FVisualLoggerFilters::Get().GetSearchString();
 
 	TArray<TFuture<void> > AllFutures;
@@ -821,8 +832,6 @@ void SVisualLogger::OnFiltersChanged()
 		FVisualLoggerDatabase::Get().SetRowVisibility(DBRow.GetOwnerName(), DBRow.GetNumberOfHiddenItems() != DBRow.GetItems().Num());
 	}
 
-
-	const uint32 EndCycles = FPlatformTime::Cycles();
 	const int32 BlockingCycles = int32(FPlatformTime::Cycles() - StartCycles);
 	{
 		const TArray<FName>& SelectedRows = FVisualLoggerDatabase::Get().GetSelectedRows();
@@ -849,7 +858,6 @@ void SVisualLogger::OnFiltersSearchChanged(const FText& Filter)
 
 	FVisualLoggerFilters::Get().SetSearchString(Filter.ToString());
 
-	ULogVisualizerSettings* Settings = ULogVisualizerSettings::StaticClass()->GetDefaultObject<ULogVisualizerSettings>();
 	const FString QuickSearchStrng = FVisualLoggerFilters::Get().GetSearchString();
 
 	TArray<TFuture<void> > AllFutures;
@@ -904,7 +912,6 @@ void SVisualLogger::OnFiltersSearchChanged(const FText& Filter)
 		VisualLoggerCanvasRenderer->DirtyCachedData();
 	}
 
-	const uint32 EndCycles = FPlatformTime::Cycles();
 	const int32 BlockingCycles = int32(FPlatformTime::Cycles() - StartCycles);
 	UE_LOG(LogVisualLogger, Display, TEXT("SVisualLogger::OnFiltersSearchChanged: %5.2fms"), FPlatformTime::ToMilliseconds(BlockingCycles));
 }
diff --git a/Engine/Source/Developer/LogVisualizer/Private/SVisualLogger.h b/Engine/Source/Developer/LogVisualizer/Private/SVisualLogger.h
--- a/Engine/Source/Developer/LogVisualizer/Private/SVisualLogger.h
+++ b/Engine/Source/Developer/LogVisualizer/Private/SVisualLogger.h
@@ -5,7 +5,6 @@
 #include "VisualLogger/VisualLoggerTypes.h"
 #include "Input/Reply.h"
 #include "Widgets/DeclarativeSyntaxSupport.h"
-#include "Engine/GameViewportClient.h"
 #include "Widgets/SCompoundWidget.h"
 #include "Framework/Commands/UICommandList.h"
 #include "Framework/Docking/TabManager.h"
@@ -43,7 +42,6 @@ public:
 	* @param InArgs The Slate argument list.
 	* @param ConstructUnderMajorTab The major tab which will contain the session front-end.
 	* @param ConstructUnderWindow The window in which this widget is being constructed.
-	* @param InStyleSet The style set to use.
 	*/
 	void Construct(const FArguments& InArgs, const TSharedRef<SDockTab>& ConstructUnderMajorTab, const TSharedPtr<SWindow>& ConstructUnderWindow);
 
@@ -127,6 +125,10 @@ protected:
 
 	bool bGotHistogramData;
 
+#if WITH_EDITOR
+	FDelegateHandle PostPIEStartedHandle;
+#endif
+
 	FDelegateHandle DrawOnCanvasDelegateHandle;
 	TWeakObjectPtr<class UWorld> LastUsedWorld;
 };
