Fix VisualLogger cylinder display

diff --git a/Engine/Source/Developer/LogVisualizer/Private/VisualLoggerRenderingActor.cpp b/Engine/Source/Developer/LogVisualizer/Private/VisualLoggerRenderingActor.cpp
--- a/Engine/Source/Developer/LogVisualizer/Private/VisualLoggerRenderingActor.cpp
+++ b/Engine/Source/Developer/LogVisualizer/Private/VisualLoggerRenderingActor.cpp
@@ -556,10 +556,12 @@ void AVisualLoggerRenderingActor::GetDebugShapes(const FVisualLogDevice::FVisual
 				const FVector Start = ElementToDraw->Points[Index];
 				const FVector End = ElementToDraw->Points[Index + 1];
 				const FVector OtherData = ElementToDraw->Points[Index + 2];
-				DebugShapes.Cylinders.Add(FDebugRenderSceneProxy::FWireCylinder(Start, OtherData.X, (End - Start).Size()*0.5, Color));
+				const float HalfHeight = 0.5f * (End - Start).Size();
+				const FVector Center = 0.5f * (Start + End);
+				DebugShapes.Cylinders.Add(FDebugRenderSceneProxy::FWireCylinder(Center, OtherData.X, HalfHeight, Color)); // Base parameter is the center of the cylinder
 				if (bDrawLabel)
 				{
-					DebugShapes.Texts.Add(FDebugRenderSceneProxy::FText3d(ElementToDraw->Description, Start, Color));
+					DebugShapes.Texts.Add(FDebugRenderSceneProxy::FText3d(ElementToDraw->Description, Center, Color));
 				}
 			}
 		}
