[VisualLogger]Refresh in Editor not running mode

diff --git a/Engine/Source/Developer/LogVisualizer/Private/SVisualLogger.cpp b/Engine/Source/Developer/LogVisualizer/Private/SVisualLogger.cpp
--- a/Engine/Source/Developer/LogVisualizer/Private/SVisualLogger.cpp
+++ b/Engine/Source/Developer/LogVisualizer/Private/SVisualLogger.cpp
@@ -208,6 +208,11 @@ void SVisualLogger::Construct(const FArguments& InArgs, const TSharedRef<SDockTa
 		FCanExecuteAction::CreateRaw(this, &SVisualLogger::HandleSaveCommandCanExecute),
 		FIsActionChecked(), 
 		FIsActionButtonVisible::CreateRaw(this, &SVisualLogger::HandleSaveCommandCanExecute));
+	ActionList.MapAction(Commands.Refresh,
+		FExecuteAction::CreateRaw(this, &SVisualLogger::HandleRefreshCommandExecute),
+		FCanExecuteAction::CreateRaw(this, &SVisualLogger::HandleRefreshCommandCanExecute),
+		FIsActionChecked(),
+		FIsActionButtonVisible::CreateRaw(this, &SVisualLogger::HandleRefreshCommandCanExecute));
 
 
 	// Tab Spawners
@@ -549,6 +554,17 @@ void SVisualLogger::HandleCameraCommandExecute()
 	}
 }
 
+bool SVisualLogger::HandleRefreshCommandCanExecute() const
+{
+	UWorld* World = FLogVisualizer::Get().GetWorld();
+	return FVisualLogger::Get().IsRecording() && World && World->IsEditorWorld();
+}
+
+void SVisualLogger::HandleRefreshCommandExecute()
+{
+	FVisualLogger::Get().Flush();
+}
+
 bool SVisualLogger::HandleLoadCommandCanExecute() const
 {
 	return true;
@@ -895,11 +911,11 @@ void SVisualLogger::OnFiltersSearchChanged(const FText& Filter)
 		}
 	} while (bAllFuturesReady != true);
 
-	for (auto Iterator = FVisualLoggerDatabase::Get().GetRowIterator(); Iterator; ++Iterator)
-	{
-		FVisualLoggerDBRow& DBRow = *Iterator;
-		FVisualLoggerDatabase::Get().SetRowVisibility(DBRow.GetOwnerName(), DBRow.GetNumberOfHiddenItems() != DBRow.GetItems().Num());
-	}
+	// for (auto Iterator = FVisualLoggerDatabase::Get().GetRowIterator(); Iterator; ++Iterator)
+	// {
+	// 	FVisualLoggerDBRow& DBRow = *Iterator;
+	// 	FVisualLoggerDatabase::Get().SetRowVisibility(DBRow.GetOwnerName(), DBRow.GetNumberOfHiddenItems() != DBRow.GetItems().Num());
+	// }
 
 	if (LogsList.IsValid())
 	{
diff --git a/Engine/Source/Developer/LogVisualizer/Private/SVisualLogger.h b/Engine/Source/Developer/LogVisualizer/Private/SVisualLogger.h
--- a/Engine/Source/Developer/LogVisualizer/Private/SVisualLogger.h
+++ b/Engine/Source/Developer/LogVisualizer/Private/SVisualLogger.h
@@ -92,6 +92,9 @@ public:
 
 	void HandleSaveCommand(bool bSaveAllData);
 
+	bool HandleRefreshCommandCanExecute() const;
+	void HandleRefreshCommandExecute();
+
 	bool HandleCameraCommandCanExecute() const;
 	void HandleCameraCommandExecute();
 	bool HandleCameraCommandIsChecked() const;
diff --git a/Engine/Source/Developer/LogVisualizer/Private/SVisualLoggerToolbar.cpp b/Engine/Source/Developer/LogVisualizer/Private/SVisualLoggerToolbar.cpp
--- a/Engine/Source/Developer/LogVisualizer/Private/SVisualLoggerToolbar.cpp
+++ b/Engine/Source/Developer/LogVisualizer/Private/SVisualLoggerToolbar.cpp
@@ -40,6 +40,7 @@ TSharedRef<SWidget> SVisualLoggerToolbar::MakeToolbar(const TSharedRef<FUIComman
 		ToolBarBuilder.AddToolBarButton(FVisualLoggerCommands::Get().FreeCamera, NAME_None, LOCTEXT("FreeCamera", "Camera"), TAttribute<FText>(), FSlateIcon(FLogVisualizerStyle::Get().GetStyleSetName(), TEXT("Toolbar.Camera")));
 		ToolBarBuilder.AddToolBarButton(FVisualLoggerCommands::Get().ResetData, NAME_None, LOCTEXT("ResetData", "Clear"), TAttribute<FText>(), FSlateIcon(FLogVisualizerStyle::Get().GetStyleSetName(), TEXT("Toolbar.Remove")));
 		ToolBarBuilder.AddToolBarButton(FVisualLoggerCommands::Get().ToggleGraphs, NAME_None, LOCTEXT("ToggleGraphs", "Graphs"), TAttribute<FText>(), FSlateIcon(FLogVisualizerStyle::Get().GetStyleSetName(), TEXT("Toolbar.Graphs")));
+		ToolBarBuilder.AddToolBarButton(FVisualLoggerCommands::Get().Refresh, NAME_None, LOCTEXT("ForceRefresh", "Refresh"), TAttribute<FText>(), FSlateIcon(FEditorStyle::GetStyleSetName(), "Icons.Refresh"));
 	}
 
 	return ToolBarBuilder.MakeWidget();
diff --git a/Engine/Source/Developer/LogVisualizer/Private/VisualLoggerCommands.h b/Engine/Source/Developer/LogVisualizer/Private/VisualLoggerCommands.h
--- a/Engine/Source/Developer/LogVisualizer/Private/VisualLoggerCommands.h
+++ b/Engine/Source/Developer/LogVisualizer/Private/VisualLoggerCommands.h
@@ -33,6 +33,7 @@ public:
 		UI_COMMAND(SaveAllToVLog, "SaveAll", "Save all to vlog file", EUserInterfaceActionType::Button, FInputChord());
 		UI_COMMAND(ToggleGraphs, "ToggleGraphs", "Toggle graphs visualization on/off", EUserInterfaceActionType::ToggleButton, FInputChord());
 		UI_COMMAND(ResetData, "Clear", "Clear all data", EUserInterfaceActionType::Button, FInputChord());
+		UI_COMMAND(Refresh, "Refresh", "Refresh UI from memory buffer", EUserInterfaceActionType::Button, FInputChord());
 		UI_COMMAND(HideEmpty, "HideEmpty", "Hide logs without text information", EUserInterfaceActionType::ToggleButton, FInputChord());
 	}
 
@@ -48,6 +49,7 @@ public:
 	TSharedPtr<FUICommandInfo> SaveAllToVLog;
 	TSharedPtr<FUICommandInfo> ToggleGraphs;
 	TSharedPtr<FUICommandInfo> ResetData;
+	TSharedPtr<FUICommandInfo> Refresh;
 	TSharedPtr<FUICommandInfo> HideEmpty;
 };
 
