[VisualLogger] Removed commented code submitted by mistake

diff --git a/Engine/Source/Developer/LogVisualizer/Private/SVisualLogger.cpp b/Engine/Source/Developer/LogVisualizer/Private/SVisualLogger.cpp
--- a/Engine/Source/Developer/LogVisualizer/Private/SVisualLogger.cpp
+++ b/Engine/Source/Developer/LogVisualizer/Private/SVisualLogger.cpp
@@ -911,11 +911,11 @@ void SVisualLogger::OnFiltersSearchChanged(const FText& Filter)
 		}
 	} while (bAllFuturesReady != true);
 
-	// for (auto Iterator = FVisualLoggerDatabase::Get().GetRowIterator(); Iterator; ++Iterator)
-	// {
-	// 	FVisualLoggerDBRow& DBRow = *Iterator;
-	// 	FVisualLoggerDatabase::Get().SetRowVisibility(DBRow.GetOwnerName(), DBRow.GetNumberOfHiddenItems() != DBRow.GetItems().Num());
-	// }
+	for (auto Iterator = FVisualLoggerDatabase::Get().GetRowIterator(); Iterator; ++Iterator)
+	{
+		FVisualLoggerDBRow& DBRow = *Iterator;
+		FVisualLoggerDatabase::Get().SetRowVisibility(DBRow.GetOwnerName(), DBRow.GetNumberOfHiddenItems() != DBRow.GetItems().Num());
+	}
 
 	if (LogsList.IsValid())
 	{
