visual logger rewind debugger integration

diff --git a/Engine/Plugins/Animation/GameplayInsights/GameplayInsights.uplugin b/Engine/Plugins/Animation/GameplayInsights/GameplayInsights.uplugin
--- a/Engine/Plugins/Animation/GameplayInsights/GameplayInsights.uplugin
+++ b/Engine/Plugins/Animation/GameplayInsights/GameplayInsights.uplugin
@@ -33,6 +33,11 @@
 			"Name": "RewindDebugger",
 			"Type": "Editor",
 			"LoadingPhase": "Default"
+		},
+		{
+			"Name": "RewindDebuggerVLog",
+			"Type": "Editor",
+			"LoadingPhase": "Default"
 		}
 
 	]
diff --git a/Engine/Plugins/Animation/GameplayInsights/Source/RewindDebuggerVLog/Private/RewindDebuggerVLog.cpp b/Engine/Plugins/Animation/GameplayInsights/Source/RewindDebuggerVLog/Private/RewindDebuggerVLog.cpp
new file mode 100644
--- /dev/null
+++ b/Engine/Plugins/Animation/GameplayInsights/Source/RewindDebuggerVLog/Private/RewindDebuggerVLog.cpp
@@ -0,0 +1,88 @@
+// Copyright Epic Games, Inc. All Rights Reserved.
+
+#include "RewindDebuggerVLog.h"
+#include "IVisualLoggerProvider.h"
+#include "Insights/IUnrealInsightsModule.h"
+#include "IRewindDebugger.h"
+#include "Modules/ModuleManager.h"
+#include "TraceServices/Model/Frames.h"
+#include "VisualLogger/VisualLogger.h"
+
+#define LOCTEXT_NAMESPACE "RewindDebuggerVLog"
+
+FRewindDebuggerVLog::FRewindDebuggerVLog()
+{
+
+}
+
+void FRewindDebuggerVLog::AddLogEntries(const TArray<TSharedPtr<FDebugObjectInfo>>& Components, float StartTime, float EndTime, const IVisualLoggerProvider* VisualLoggerProvider)
+{
+	for(const TSharedPtr<FDebugObjectInfo>& ComponentInfo : Components)
+	{
+		VisualLoggerProvider->ReadVisualLogEntryTimeline(ComponentInfo->ObjectId, [this,StartTime, EndTime](const IVisualLoggerProvider::VisualLogEntryTimeline &TimelineData)
+		{
+			TimelineData.EnumerateEvents(StartTime, EndTime, [this](double InStartTime, double InEndTime, uint32 InDepth, const FVisualLogEntry& LogEntry)
+			{
+				VLogActor->AddLogEntry(LogEntry);
+				return TraceServices::EEventEnumerate::Continue;
+			});
+		});
+
+		AddLogEntries(ComponentInfo->Children, StartTime, EndTime, VisualLoggerProvider);
+	}
+}
+
+void FRewindDebuggerVLog::Update(float DeltaTime, IRewindDebugger* RewindDebugger)
+{
+	if (RewindDebugger->IsPIESimulating() && !RewindDebugger->IsRecording())
+	{
+		// output debug rendering if we are paused, or actively recording.
+		// otherwise clear the debug rendering actor to avoid any debug rendering being left over from last time we were scrubbing or recording
+		if (VLogActor.IsValid())
+		{
+			VLogActor->Reset();
+		}
+	}
+	else
+	{
+		if (const TraceServices::IAnalysisSession* Session = RewindDebugger->GetAnalysisSession())
+		{
+			TraceServices::FAnalysisSessionReadScope SessionReadScope(*Session);
+			double CurrentTraceTime = RewindDebugger->CurrentTraceTime();
+
+			const TraceServices::IFrameProvider& FrameProvider = TraceServices::ReadFrameProvider(*Session);
+			TraceServices::FFrame CurrentFrame;
+
+			if(FrameProvider.GetFrameFromTime(ETraceFrameType::TraceFrameType_Game, CurrentTraceTime, CurrentFrame))
+			{
+				if (const IVisualLoggerProvider* VisualLoggerProvider = Session->ReadProvider<IVisualLoggerProvider>("VisualLoggerProvider"))
+				{
+					if (!VLogActor.IsValid())
+					{
+						FActorSpawnParameters SpawnParameters;
+						SpawnParameters.ObjectFlags |= RF_Transient;
+						VLogActor = RewindDebugger->GetWorldToVisualize()->SpawnActor<AVLogRenderingActor>(SpawnParameters);
+					}
+
+					VLogActor->Reset();
+
+					AddLogEntries(RewindDebugger->GetDebugComponents(), CurrentFrame.StartTime, CurrentFrame.EndTime, VisualLoggerProvider);
+				}
+			}
+		}
+	}
+}
+
+void FRewindDebuggerVLog::RecordingStarted(IRewindDebugger*)
+{
+	// start recording visual logger data
+	FVisualLogger::Get().SetIsRecordingToTrace(true);
+}
+
+void FRewindDebuggerVLog::RecordingStopped(IRewindDebugger*)
+{
+	// stop recording visual logger data
+	FVisualLogger::Get().SetIsRecordingToTrace(false);
+}
+
+#undef LOCTEXT_NAMESPACE
\ No newline at end of file
diff --git a/Engine/Plugins/Animation/GameplayInsights/Source/RewindDebuggerVLog/Private/RewindDebuggerVLog.h b/Engine/Plugins/Animation/GameplayInsights/Source/RewindDebuggerVLog/Private/RewindDebuggerVLog.h
new file mode 100644
--- /dev/null
+++ b/Engine/Plugins/Animation/GameplayInsights/Source/RewindDebuggerVLog/Private/RewindDebuggerVLog.h
@@ -0,0 +1,27 @@
+// Copyright Epic Games, Inc. All Rights Reserved.
+#pragma once
+
+#include "CoreMinimal.h"
+#include "IRewindDebugger.h"
+#include "IRewindDebuggerExtension.h"
+#include "UObject/WeakObjectPtr.h"
+#include "VLogRenderingActor.h"
+
+// Rewind debugger extension for Visual Logger support
+
+class FRewindDebuggerVLog : public IRewindDebuggerExtension
+{
+public:
+	FRewindDebuggerVLog();
+	virtual ~FRewindDebuggerVLog() {};
+
+	virtual void Update(float DeltaTime, IRewindDebugger* RewindDebugger) override;
+
+	virtual void RecordingStarted(IRewindDebugger* RewindDebugger) override;
+	virtual void RecordingStopped(IRewindDebugger* RewindDebugger) override;
+
+private:
+	void AddLogEntries(const TArray<TSharedPtr<FDebugObjectInfo>>& Components, float StartTime, float EndTime, const class IVisualLoggerProvider* Provider);
+
+	TWeakObjectPtr<AVLogRenderingActor> VLogActor; 
+};
diff --git a/Engine/Plugins/Animation/GameplayInsights/Source/RewindDebuggerVLog/Private/RewindDebuggerVLogModule.cpp b/Engine/Plugins/Animation/GameplayInsights/Source/RewindDebuggerVLog/Private/RewindDebuggerVLogModule.cpp
new file mode 100644
--- /dev/null
+++ b/Engine/Plugins/Animation/GameplayInsights/Source/RewindDebuggerVLog/Private/RewindDebuggerVLogModule.cpp
@@ -0,0 +1,22 @@
+// Copyright Epic Games, Inc. All Rights Reserved.
+
+#include "RewindDebuggerVLogModule.h"
+#include "Features/IModularFeatures.h"
+
+#define LOCTEXT_NAMESPACE "RewindDebuggerVLogModule"
+
+void FRewindDebuggerVLogModule::StartupModule()
+{
+	IModularFeatures::Get().RegisterModularFeature(IRewindDebuggerExtension::ModularFeatureName, &RewindDebuggerVLogExtension);
+	IModularFeatures::Get().RegisterModularFeature(TraceServices::ModuleFeatureName, &VLogTraceModule);
+}
+
+void FRewindDebuggerVLogModule::ShutdownModule()
+{
+	IModularFeatures::Get().UnregisterModularFeature(IRewindDebuggerExtension::ModularFeatureName, &RewindDebuggerVLogExtension);
+	IModularFeatures::Get().UnregisterModularFeature(TraceServices::ModuleFeatureName, &VLogTraceModule);
+}
+
+IMPLEMENT_MODULE(FRewindDebuggerVLogModule, RewindDebuggerVLog);
+
+#undef LOCTEXT_NAMESPACE
diff --git a/Engine/Plugins/Animation/GameplayInsights/Source/RewindDebuggerVLog/Private/RewindDebuggerVLogModule.h b/Engine/Plugins/Animation/GameplayInsights/Source/RewindDebuggerVLog/Private/RewindDebuggerVLogModule.h
new file mode 100644
--- /dev/null
+++ b/Engine/Plugins/Animation/GameplayInsights/Source/RewindDebuggerVLog/Private/RewindDebuggerVLogModule.h
@@ -0,0 +1,22 @@
+// Copyright Epic Games, Inc. All Rights Reserved.
+
+#pragma once
+
+#include "CoreMinimal.h"
+#include "Modules/ModuleInterface.h"
+#include "RewindDebuggerVLog.h"
+#include "VLogTraceModule.h"
+
+class SRewindDebugger;
+
+class FRewindDebuggerVLogModule : public IModuleInterface
+{
+public:
+	// IModuleInterface interface
+	virtual void StartupModule() override;
+	virtual void ShutdownModule() override;
+
+private:
+	FRewindDebuggerVLog RewindDebuggerVLogExtension;
+	FVLogTraceModule VLogTraceModule;
+};
diff --git a/Engine/Plugins/Animation/GameplayInsights/Source/RewindDebuggerVLog/Private/VLogRenderingActor.cpp b/Engine/Plugins/Animation/GameplayInsights/Source/RewindDebuggerVLog/Private/VLogRenderingActor.cpp
new file mode 100644
--- /dev/null
+++ b/Engine/Plugins/Animation/GameplayInsights/Source/RewindDebuggerVLog/Private/VLogRenderingActor.cpp
@@ -0,0 +1,27 @@
+// Copyright Epic Games, Inc. All Rights Reserved.
+
+#include "VLogRenderingActor.h"
+
+AVLogRenderingActor::AVLogRenderingActor(const FObjectInitializer& ObjectInitializer)
+{
+}
+
+AVLogRenderingActor::~AVLogRenderingActor()
+{
+}
+
+void AVLogRenderingActor::Reset()
+{
+ 	DebugShapes.Reset();
+ 	MarkComponentsRenderStateDirty();
+}
+
+void AVLogRenderingActor::AddLogEntry(const FVisualLogEntry& Entry)
+{
+	GetDebugShapes(Entry, false, DebugShapes);
+}
+
+void AVLogRenderingActor::IterateDebugShapes(TFunction<void(const AVisualLoggerRenderingActorBase::FTimelineDebugShapes& Shapes)> Callback)
+{
+	Callback(DebugShapes);
+}
diff --git a/Engine/Plugins/Animation/GameplayInsights/Source/RewindDebuggerVLog/Private/VLogRenderingActor.h b/Engine/Plugins/Animation/GameplayInsights/Source/RewindDebuggerVLog/Private/VLogRenderingActor.h
new file mode 100644
--- /dev/null
+++ b/Engine/Plugins/Animation/GameplayInsights/Source/RewindDebuggerVLog/Private/VLogRenderingActor.h
@@ -0,0 +1,27 @@
+// Copyright Epic Games, Inc. All Rights Reserved.
+
+#pragma once
+
+#include "CoreMinimal.h"
+#include "UObject/ObjectMacros.h"
+#include "VisualLoggerRenderingActorBase.h"
+#include "VLogRenderingActor.generated.h"
+
+/**
+*	Transient actor used to draw visual logger data on level
+*/
+
+UCLASS(config = Engine, NotBlueprintable, Transient, notplaceable, AdvancedClassDisplay)
+class AVLogRenderingActor : public AVisualLoggerRenderingActorBase 
+{
+public:
+	GENERATED_UCLASS_BODY()
+	virtual ~AVLogRenderingActor();
+	void Reset();
+	void AddLogEntry(const FVisualLogEntry& LogEntry);
+
+	virtual void IterateDebugShapes(TFunction<void(const FTimelineDebugShapes&) > Callback) override;
+
+private:
+	FTimelineDebugShapes DebugShapes;
+};
diff --git a/Engine/Plugins/Animation/GameplayInsights/Source/RewindDebuggerVLog/Private/VLogTraceModule.cpp b/Engine/Plugins/Animation/GameplayInsights/Source/RewindDebuggerVLog/Private/VLogTraceModule.cpp
new file mode 100644
--- /dev/null
+++ b/Engine/Plugins/Animation/GameplayInsights/Source/RewindDebuggerVLog/Private/VLogTraceModule.cpp
@@ -0,0 +1,32 @@
+// Copyright Epic Games, Inc. All Rights Reserved.
+
+#include "VLogTraceModule.h"
+#include "VisualLoggerProvider.h"
+#include "VisualLoggerAnalyzer.h"
+
+FName FVLogTraceModule::ModuleName("GameplayTrace");
+
+void FVLogTraceModule::GetModuleInfo(TraceServices::FModuleInfo& OutModuleInfo)
+{
+	OutModuleInfo.Name = ModuleName;
+	OutModuleInfo.DisplayName = TEXT("VisualLogger");
+}
+
+void FVLogTraceModule::OnAnalysisBegin(TraceServices::IAnalysisSession& InSession)
+{
+	FVisualLoggerProvider* VisualLoggerProvider = new FVisualLoggerProvider(InSession);
+	InSession.AddProvider(FVisualLoggerProvider::ProviderName, VisualLoggerProvider);
+
+	InSession.AddAnalyzer(new FVisualLoggerAnalyzer(InSession, *VisualLoggerProvider));
+}
+
+void FVLogTraceModule::GetLoggers(TArray<const TCHAR *>& OutLoggers)
+{
+	OutLoggers.Add(TEXT("VisualLogger"));
+}
+
+void FVLogTraceModule::GenerateReports(const TraceServices::IAnalysisSession& Session, const TCHAR* CmdLine, const TCHAR* OutputDirectory)
+{
+
+}
+
diff --git a/Engine/Plugins/Animation/GameplayInsights/Source/RewindDebuggerVLog/Private/VLogTraceModule.h b/Engine/Plugins/Animation/GameplayInsights/Source/RewindDebuggerVLog/Private/VLogTraceModule.h
new file mode 100644
--- /dev/null
+++ b/Engine/Plugins/Animation/GameplayInsights/Source/RewindDebuggerVLog/Private/VLogTraceModule.h
@@ -0,0 +1,24 @@
+// Copyright Epic Games, Inc. All Rights Reserved.
+
+#pragma once
+
+#include "CoreMinimal.h"
+#include "TraceServices/ModuleService.h"
+
+class IAnimationProvider;
+namespace Insights { class ITimingViewSession; }
+
+class FVLogTraceModule : public TraceServices::IModule
+{
+public:
+	// TraceServices::IModule interface
+	virtual void GetModuleInfo(TraceServices::FModuleInfo& OutModuleInfo) override;
+	virtual void OnAnalysisBegin(TraceServices::IAnalysisSession& Session) override;
+	virtual void GetLoggers(TArray<const TCHAR *>& OutLoggers) override;
+	virtual void GenerateReports(const TraceServices::IAnalysisSession& Session, const TCHAR* CmdLine, const TCHAR* OutputDirectory) override;
+	virtual const TCHAR* GetCommandLineArgument() override { return TEXT("visuallogger"); }
+
+private:
+	static FName ModuleName;
+};
+
diff --git a/Engine/Plugins/Animation/GameplayInsights/Source/RewindDebuggerVLog/Private/VisualLoggerAnalyzer.cpp b/Engine/Plugins/Animation/GameplayInsights/Source/RewindDebuggerVLog/Private/VisualLoggerAnalyzer.cpp
new file mode 100644
--- /dev/null
+++ b/Engine/Plugins/Animation/GameplayInsights/Source/RewindDebuggerVLog/Private/VisualLoggerAnalyzer.cpp
@@ -0,0 +1,42 @@
+// Copyright Epic Games, Inc. All Rights Reserved.
+
+#include "VisualLoggerAnalyzer.h"
+#include "TraceServices/Model/AnalysisSession.h"
+#include "VisualLoggerProvider.h"
+#include "Serialization/MemoryReader.h"
+
+FVisualLoggerAnalyzer::FVisualLoggerAnalyzer(TraceServices::IAnalysisSession& InSession, FVisualLoggerProvider& InProvider)
+	: Session(InSession)
+	, Provider(InProvider)
+{
+}
+
+void FVisualLoggerAnalyzer::OnAnalysisBegin(const FOnAnalysisContext& Context)
+{
+	auto& Builder = Context.InterfaceBuilder;
+
+	Builder.RouteEvent(RouteId_VisualLogEntry, "VisualLogger", "VisualLogEntry");
+}
+
+bool FVisualLoggerAnalyzer::OnEvent(uint16 RouteId, EStyle Style, const FOnEventContext& Context)
+{
+	TraceServices::FAnalysisSessionEditScope _(Session);
+
+	const auto& EventData = Context.EventData;
+	switch (RouteId)
+	{
+		case RouteId_VisualLogEntry:
+		{
+			uint64 OwnerId = EventData.GetValue<uint64>("OwnerId");
+			uint64 Cycle = EventData.GetValue<uint64>("Cycle");
+			TArrayView<const uint8> SerializedData = EventData.GetArrayView<uint8>("LogEntry");
+			FMemoryReaderView Archive(SerializedData);
+			FVisualLogEntry Entry;
+			Archive << Entry;
+			Provider.AppendVisualLogEntry(OwnerId, Context.EventTime.AsSeconds(Cycle), Entry);
+			break;
+		}
+	}
+
+	return true;
+}
diff --git a/Engine/Plugins/Animation/GameplayInsights/Source/RewindDebuggerVLog/Private/VisualLoggerAnalyzer.h b/Engine/Plugins/Animation/GameplayInsights/Source/RewindDebuggerVLog/Private/VisualLoggerAnalyzer.h
new file mode 100644
--- /dev/null
+++ b/Engine/Plugins/Animation/GameplayInsights/Source/RewindDebuggerVLog/Private/VisualLoggerAnalyzer.h
@@ -0,0 +1,27 @@
+// Copyright Epic Games, Inc. All Rights Reserved.
+
+#pragma once
+
+#include "Trace/Analyzer.h"
+
+class FVisualLoggerProvider;
+namespace TraceServices { class IAnalysisSession; }
+
+class FVisualLoggerAnalyzer : public UE::Trace::IAnalyzer
+{
+public:
+	FVisualLoggerAnalyzer(TraceServices::IAnalysisSession& InSession, FVisualLoggerProvider& InProvider);
+
+	virtual void OnAnalysisBegin(const FOnAnalysisContext& Context) override;
+	virtual void OnAnalysisEnd() override {}
+	virtual bool OnEvent(uint16 RouteId, EStyle Style, const FOnEventContext& Context) override;
+
+private:
+	enum : uint16
+	{
+		RouteId_VisualLogEntry,
+	};
+
+	TraceServices::IAnalysisSession& Session;
+	FVisualLoggerProvider& Provider;
+};
diff --git a/Engine/Plugins/Animation/GameplayInsights/Source/RewindDebuggerVLog/Private/VisualLoggerProvider.cpp b/Engine/Plugins/Animation/GameplayInsights/Source/RewindDebuggerVLog/Private/VisualLoggerProvider.cpp
new file mode 100644
--- /dev/null
+++ b/Engine/Plugins/Animation/GameplayInsights/Source/RewindDebuggerVLog/Private/VisualLoggerProvider.cpp
@@ -0,0 +1,51 @@
+// Copyright Epic Games, Inc. All Rights Reserved.
+
+#include "VisualLoggerProvider.h"
+
+FName FVisualLoggerProvider::ProviderName("VisualLoggerProvider");
+
+#define LOCTEXT_NAMESPACE "VisualLoggerProvider"
+
+FVisualLoggerProvider::FVisualLoggerProvider(TraceServices::IAnalysisSession& InSession)
+	: Session(InSession)
+{
+}
+
+bool FVisualLoggerProvider::ReadVisualLogEntryTimeline(uint64 InObjectId, TFunctionRef<void(const VisualLogEntryTimeline&)> Callback) const
+{
+	Session.ReadAccessCheck();
+
+	const uint32* IndexPtr = ObjectIdToLogEntryTimelines.Find(InObjectId);
+	if(IndexPtr != nullptr)
+	{
+		if (*IndexPtr < uint32(LogEntryTimelines.Num()))
+		{
+			Callback(*LogEntryTimelines[*IndexPtr]);
+			return true;
+		}
+	}
+
+	return false;
+}
+
+void  FVisualLoggerProvider::AppendVisualLogEntry(uint64 InObjectId, double InTime, const FVisualLogEntry& Entry)
+{
+	Session.WriteAccessCheck();
+
+	TSharedPtr<TraceServices::TPointTimeline<FVisualLogEntry>> Timeline;
+	uint32* IndexPtr = ObjectIdToLogEntryTimelines.Find(InObjectId);
+	if(IndexPtr != nullptr)
+	{
+		Timeline = LogEntryTimelines[*IndexPtr];
+	}
+	else
+	{
+		Timeline = MakeShared<TraceServices::TPointTimeline<FVisualLogEntry>>(Session.GetLinearAllocator());
+		ObjectIdToLogEntryTimelines.Add(InObjectId, LogEntryTimelines.Num());
+		LogEntryTimelines.Add(Timeline.ToSharedRef());
+	}
+
+	Timeline->AppendEvent(InTime, Entry);
+}
+
+#undef LOCTEXT_NAMESPACE
diff --git a/Engine/Plugins/Animation/GameplayInsights/Source/RewindDebuggerVLog/Private/VisualLoggerProvider.h b/Engine/Plugins/Animation/GameplayInsights/Source/RewindDebuggerVLog/Private/VisualLoggerProvider.h
new file mode 100644
--- /dev/null
+++ b/Engine/Plugins/Animation/GameplayInsights/Source/RewindDebuggerVLog/Private/VisualLoggerProvider.h
@@ -0,0 +1,32 @@
+// Copyright Epic Games, Inc. All Rights Reserved.
+
+#pragma once
+
+#include "IVisualLoggerProvider.h"
+#include "Model/PointTimeline.h"
+#include "Containers/Map.h"
+#include "Templates/SharedPointer.h"
+#include "Model/IntervalTimeline.h"
+#include "VisualLogger/VisualLoggerTypes.h"
+
+namespace TraceServices { class IAnalysisSession; }
+
+class FVisualLoggerProvider : public IVisualLoggerProvider
+{
+public:
+	static FName ProviderName;
+
+	FVisualLoggerProvider(TraceServices::IAnalysisSession& InSession);
+
+	/** IVisualLoggerProvider interface */
+	virtual bool ReadVisualLogEntryTimeline(uint64 InObjectId, TFunctionRef<void(const VisualLogEntryTimeline&)> Callback) const;
+
+	/** Add an object event message */
+	void AppendVisualLogEntry(uint64 InObjectId, double InTime, const FVisualLogEntry& Entry);
+
+private:
+	TraceServices::IAnalysisSession& Session;
+
+	TMap<uint64, uint32> ObjectIdToLogEntryTimelines;
+	TArray<TSharedRef<TraceServices::TPointTimeline<FVisualLogEntry>>> LogEntryTimelines;
+};
diff --git a/Engine/Plugins/Animation/GameplayInsights/Source/RewindDebuggerVLog/Public/IVisualLoggerProvider.h b/Engine/Plugins/Animation/GameplayInsights/Source/RewindDebuggerVLog/Public/IVisualLoggerProvider.h
new file mode 100644
--- /dev/null
+++ b/Engine/Plugins/Animation/GameplayInsights/Source/RewindDebuggerVLog/Public/IVisualLoggerProvider.h
@@ -0,0 +1,17 @@
+// Copyright Epic Games, Inc. All Rights Reserved.
+
+#pragma once
+
+#include "CoreMinimal.h"
+#include "TraceServices/Model/AnalysisSession.h"
+#include "TraceServices/Containers/Timelines.h"
+#include "VisualLogger/VisualLoggerTypes.h"
+
+class IVisualLoggerProvider : public TraceServices::IProvider
+{
+public:
+	typedef TraceServices::ITimeline<FVisualLogEntry> VisualLogEntryTimeline;
+
+	virtual bool ReadVisualLogEntryTimeline(uint64 InObjectId, TFunctionRef<void(const VisualLogEntryTimeline&)> Callback) const = 0;
+};
+
diff --git a/Engine/Plugins/Animation/GameplayInsights/Source/RewindDebuggerVLog/RewindDebuggerVLog.Build.cs b/Engine/Plugins/Animation/GameplayInsights/Source/RewindDebuggerVLog/RewindDebuggerVLog.Build.cs
new file mode 100644
--- /dev/null
+++ b/Engine/Plugins/Animation/GameplayInsights/Source/RewindDebuggerVLog/RewindDebuggerVLog.Build.cs
@@ -0,0 +1,36 @@
+// Copyright Epic Games, Inc. All Rights Reserved.
+
+namespace UnrealBuildTool.Rules
+{
+	public class RewindDebuggerVLog : ModuleRules
+	{
+		public RewindDebuggerVLog(ReadOnlyTargetRules Target) : base(Target)
+		{
+			PrivateDependencyModuleNames.AddRange(
+			new string[]
+			{
+				"ApplicationCore",
+				"AssetRegistry",
+				"Core",
+				"CoreUObject",
+				"EditorFramework",
+				"EditorStyle",
+				"EditorWidgets",
+				"Engine",
+				"GameplayInsights",
+				"GameplayInsightsEditor",
+				"LogVisualizer",
+				"SlateCore",
+				"Slate",
+				"ToolMenus",
+				"TraceLog",
+				"TraceAnalysis",
+				"TraceServices",
+				"TraceInsights",
+				"UnrealEd",
+				"RewindDebuggerInterface"
+			});
+		}
+	}
+}
+
diff --git a/Engine/Source/Developer/LogVisualizer/Private/VisualLoggerRenderingActor.cpp b/Engine/Source/Developer/LogVisualizer/Private/VisualLoggerRenderingActor.cpp
--- a/Engine/Source/Developer/LogVisualizer/Private/VisualLoggerRenderingActor.cpp
+++ b/Engine/Source/Developer/LogVisualizer/Private/VisualLoggerRenderingActor.cpp
@@ -11,135 +11,9 @@
 #endif // WITH_EDITOR
 #include "VisualLoggerRenderingComponent.h"
 
-namespace FDebugDrawing
-{
-	const FVector NavOffset(0, 0, 15);
-}
-
-class UVisualLoggerRenderingComponent;
-class FVisualLoggerSceneProxy final : public FDebugRenderSceneProxy
-{
-public:
-	SIZE_T GetTypeHash() const override
-	{
-		static size_t UniquePointer;
-		return reinterpret_cast<size_t>(&UniquePointer);
-	}
-
-	FVisualLoggerSceneProxy(const UVisualLoggerRenderingComponent* InComponent)
-		: FDebugRenderSceneProxy(InComponent)
-	{
-		DrawType = SolidAndWireMeshes;
-		ViewFlagName = TEXT("VisLog");
-		ViewFlagIndex = uint32(FEngineShowFlags::FindIndexByName(*ViewFlagName));
-		bWantsSelectionOutline = false;
-	}
-
-	virtual FPrimitiveViewRelevance GetViewRelevance(const FSceneView* View) const override
-	{
-		FPrimitiveViewRelevance Result;
-		Result.bDrawRelevance = IsShown(View);
-		Result.bDynamicRelevance = true;
-		// ideally the TranslucencyRelevance should be filled out by the material, here we do it conservative
-		Result.bSeparateTranslucency = Result.bNormalTranslucency = IsShown(View) && GIsEditor;
-		return Result;
-	}
-
-	virtual uint32 GetMemoryFootprint(void) const override { return sizeof(*this) + GetAllocatedSize(); }
-
-	uint32 GetAllocatedSize(void) const
-	{
-		return FDebugRenderSceneProxy::GetAllocatedSize();
-	}
-};
-
-UVisualLoggerRenderingComponent::UVisualLoggerRenderingComponent(const FObjectInitializer& ObjectInitializer) : Super(ObjectInitializer)
-{
-}
-
-FPrimitiveSceneProxy* UVisualLoggerRenderingComponent::CreateSceneProxy()
-{
-	AVisualLoggerRenderingActor* RenderingActor = Cast<AVisualLoggerRenderingActor>(GetOuter());
-	if (RenderingActor == NULL)
-	{
-		return NULL;
-	}
-
-	ULogVisualizerSettings *Settings = ULogVisualizerSettings::StaticClass()->GetDefaultObject<ULogVisualizerSettings>();
-	FVisualLoggerSceneProxy *VLogSceneProxy = new FVisualLoggerSceneProxy(this);
-	VLogSceneProxy->SolidMeshMaterial = Settings->GetDebugMeshMaterial();
-
-	for (auto& CurrentShapes : RenderingActor->DebugShapesPerRow)
-	{
-		VLogSceneProxy->Spheres.Append(CurrentShapes.Value.Points);
-		VLogSceneProxy->Lines.Append(CurrentShapes.Value.Lines);
-		VLogSceneProxy->Boxes.Append(CurrentShapes.Value.Boxes);
-		VLogSceneProxy->Meshes.Append(CurrentShapes.Value.Meshes);
-		VLogSceneProxy->Cones.Append(CurrentShapes.Value.Cones);
-		VLogSceneProxy->Texts.Append(CurrentShapes.Value.Texts);
-		VLogSceneProxy->Cylinders.Append(CurrentShapes.Value.Cylinders);
-		VLogSceneProxy->ArrowLines.Append(CurrentShapes.Value.Arrows);
-		VLogSceneProxy->Capsules.Append(CurrentShapes.Value.Capsules);
-	}
-
-	{
-		VLogSceneProxy->Spheres.Append(RenderingActor->TestDebugShapes.Points);
-		VLogSceneProxy->Lines.Append(RenderingActor->TestDebugShapes.Lines);
-		VLogSceneProxy->Boxes.Append(RenderingActor->TestDebugShapes.Boxes);
-		VLogSceneProxy->Meshes.Append(RenderingActor->TestDebugShapes.Meshes);
-		VLogSceneProxy->Cones.Append(RenderingActor->TestDebugShapes.Cones);
-		VLogSceneProxy->Texts.Append(RenderingActor->TestDebugShapes.Texts);
-		VLogSceneProxy->Cylinders.Append(RenderingActor->TestDebugShapes.Cylinders);
-		VLogSceneProxy->ArrowLines.Append(RenderingActor->TestDebugShapes.Arrows);
-		VLogSceneProxy->Capsules.Append(RenderingActor->TestDebugShapes.Capsules);
-	}
-
-#if WITH_EDITOR
-	if (VLogSceneProxy)
-	{
-		DebugDrawDelegateHelper.InitDelegateHelper(VLogSceneProxy);
-		DebugDrawDelegateHelper.ReregisterDebugDrawDelegate();
-	}
-#endif
-	return VLogSceneProxy;
-}
-
-FBoxSphereBounds UVisualLoggerRenderingComponent::CalcBounds(const FTransform& LocalToWorld) const
-{
-	FBox MyBounds;
-	MyBounds.Init();
-
-	MyBounds = FBox(FVector(-HALF_WORLD_MAX, -HALF_WORLD_MAX, -HALF_WORLD_MAX), FVector(HALF_WORLD_MAX, HALF_WORLD_MAX, HALF_WORLD_MAX));
-
-	return MyBounds;
-}
-
-void UVisualLoggerRenderingComponent::CreateRenderState_Concurrent(FRegisterComponentContext* Context)
-{
-	Super::CreateRenderState_Concurrent(Context);
-
-#if WITH_EDITOR
-	DebugDrawDelegateHelper.RegisterDebugDrawDelegate();
-#endif
-}
-
-void UVisualLoggerRenderingComponent::DestroyRenderState_Concurrent()
-{
-#if WITH_EDITOR
-	DebugDrawDelegateHelper.UnregisterDebugDrawDelegate();
-#endif
-
-	Super::DestroyRenderState_Concurrent();
-}
-
 AVisualLoggerRenderingActor::AVisualLoggerRenderingActor(const FObjectInitializer& ObjectInitializer)
 	: Super(ObjectInitializer)
 {
-	USceneComponent* SceneComponent = CreateDefaultSubobject<USceneComponent>(TEXT("SceneComp"));
-	RootComponent = SceneComponent;
-
-	RenderingComponent = CreateDefaultSubobject<UVisualLoggerRenderingComponent>(TEXT("RenderingComponent"));
-
 	if (HasAnyFlags(RF_ClassDefaultObject) == false)
 	{
 		FVisualLoggerDatabase::Get().GetEvents().OnItemSelectionChanged.AddUObject(this, &AVisualLoggerRenderingActor::OnItemSelectionChanged);
@@ -147,6 +21,10 @@ AVisualLoggerRenderingActor::AVisualLoggerRenderingActor(const FObjectInitialize
 
 		FLogVisualizer::Get().GetEvents().OnFiltersChanged.AddUObject(this, &AVisualLoggerRenderingActor::OnFiltersChanged);
 	}
+
+#if VLOG_TEST_DEBUG_RENDERING
+	AddDebugRendering();
+#endif
 }
 
 AVisualLoggerRenderingActor::~AVisualLoggerRenderingActor()
@@ -220,7 +98,7 @@ void AVisualLoggerRenderingActor::OnItemSelectionChanged(const FVisualLoggerDBRo
 	ShapesCache.Reset();
 	const TArray<FVisualLogDevice::FVisualLogEntryItem>& Entries = DBRow.GetItems();
 	const int32 CurrentItemIndex = DBRow.GetCurrentItemIndex();
-	GetDebugShapes(Entries[CurrentItemIndex], ShapesCache);
+	GetDebugShapes(Entries[CurrentItemIndex].Entry, true, ShapesCache);
 
 	MarkComponentsRenderStateDirty();
 }
@@ -249,12 +127,13 @@ void AVisualLoggerRenderingActor::OnFiltersChanged()
 		ShapesCache.Reset();
 		if (DBRow.GetCurrentItemIndex() != INDEX_NONE)
 		{
-			GetDebugShapes(DBRow.GetCurrentItem(), ShapesCache);
+			GetDebugShapes(DBRow.GetCurrentItem().Entry, true, ShapesCache);
 		}
 	}
 	MarkComponentsRenderStateDirty();
 }
 
+#if VLOG_TEST_DEBUG_RENDERING
 void AVisualLoggerRenderingActor::AddDebugRendering()
 {
 	const float Thickness = 2;
@@ -302,422 +181,16 @@ void AVisualLoggerRenderingActor::AddDebugRendering()
 		TestDebugShapes.Points.Add(FDebugRenderSceneProxy::FSphere(10, FVector(1300, 0, 128), FColor::White));
 	}
 }
+#endif
 
-namespace
+void AVisualLoggerRenderingActor::IterateDebugShapes(TFunction<void(const AVisualLoggerRenderingActorBase::FTimelineDebugShapes& Shapes)> Callback)
 {
-	static bool IsPolygonWindingCorrect(const TArray<FVector>& Verts)
+	for (auto& CurrentShapes : DebugShapesPerRow)
 	{
-		// this will work only for convex polys, but we're assuming that all logged polygons are convex in the first place
-		if (Verts.Num() >= 3)
-		{
-			const FVector SurfaceNormal = FVector::CrossProduct(Verts[1] - Verts[0], Verts[2] - Verts[0]);
-			const float TestDot = FVector::DotProduct(SurfaceNormal, FVector(0, 0, 1));
-			return TestDot > 0;
-		}
-
-		return false;
+		Callback(CurrentShapes.Value);
 	}
 
-	static void GetPolygonMesh(const FVisualLogShapeElement* ElementToDraw, FDebugRenderSceneProxy::FMesh& TestMesh, const FVector3f& VertexOffset = FVector3f::ZeroVector)
-	{
-		TestMesh.Color = ElementToDraw->GetFColor();
-
-		FClipSMPolygon InPoly(ElementToDraw->Points.Num());
-		InPoly.FaceNormal = FVector(0, 0, 1);
-
-		const bool bHasCorrectWinding = IsPolygonWindingCorrect(ElementToDraw->Points);
-		if (bHasCorrectWinding)
-		{
-			for (int32 Index = 0; Index < ElementToDraw->Points.Num(); Index++)
-			{
-				FClipSMVertex v1;
-				v1.Pos = ElementToDraw->Points[Index];
-				InPoly.Vertices.Add(v1);
-			}
-		}
-		else
-		{
-			for (int32 Index = ElementToDraw->Points.Num() - 1; Index >= 0; Index--)
-			{
-				FClipSMVertex v1;
-				v1.Pos = ElementToDraw->Points[Index];
-				InPoly.Vertices.Add(v1);
-			}
-		}
-
-		TArray<FClipSMTriangle> OutTris;
-		
-		const bool bTriangulated = FGeomTools::TriangulatePoly(OutTris, InPoly, false);
-		if (bTriangulated)
-		{
-			int32 LastIndex = 0;
-
-			FGeomTools::RemoveRedundantTriangles(OutTris);
-			for (const auto& CurrentTri : OutTris)
-			{
-				TestMesh.Vertices.Add(FDynamicMeshVertex(CurrentTri.Vertices[0].Pos + VertexOffset));
-				TestMesh.Vertices.Add(FDynamicMeshVertex(CurrentTri.Vertices[1].Pos + VertexOffset));
-				TestMesh.Vertices.Add(FDynamicMeshVertex(CurrentTri.Vertices[2].Pos + VertexOffset));
-				TestMesh.Indices.Add(LastIndex++);
-				TestMesh.Indices.Add(LastIndex++);
-				TestMesh.Indices.Add(LastIndex++);
-			}
-		}
-	}
-}
-
-void AVisualLoggerRenderingActor::GetDebugShapes(const FVisualLogDevice::FVisualLogEntryItem& EntryItem, FTimelineDebugShapes& DebugShapes)
-{
-	const FVisualLogEntry* Entry = &EntryItem.Entry;
-	const FVisualLogShapeElement* ElementToDraw = Entry->ElementsToDraw.GetData();
-	const int32 ElementsCount = Entry->ElementsToDraw.Num();
-
-#if 0
-	AddDebugRendering();
+#if VLOG_TEST_DEBUG_RENDERING
+	Callback(TestDebugShapes);
 #endif
-
-	{
-		const float Length = 100;
-		const FVector DirectionNorm = FVector(0, 0, 1).GetSafeNormal();
-		FVector YAxis, ZAxis;
-		DirectionNorm.FindBestAxisVectors(YAxis, ZAxis);
-		DebugShapes.Cones.Add(FDebugRenderSceneProxy::FCone(FScaleMatrix(FVector(Length)) * FMatrix(DirectionNorm, YAxis, ZAxis, Entry->Location), 5, 5, FColor::Red));
-	}
-
-	if (DebugShapes.LogEntriesPath.Num())
-	{
-		FVector Location = DebugShapes.LogEntriesPath[0];
-		for (int32 Index = 1; Index < DebugShapes.LogEntriesPath.Num(); ++Index)
-		{
-			const FVector CurrentLocation = DebugShapes.LogEntriesPath[Index];
-			DebugShapes.Lines.Add(FDebugRenderSceneProxy::FDebugLine(Location, CurrentLocation, FColor(160, 160, 240), 2.0));
-			Location = CurrentLocation;
-		}
-	}
-
-	for (int32 ElementIndex = 0; ElementIndex < ElementsCount; ++ElementIndex, ++ElementToDraw)
-	{
-		if (!FVisualLoggerFilters::Get().MatchCategoryFilters(ElementToDraw->Category.ToString(), ElementToDraw->Verbosity))
-		{
-			continue;
-		}
-
-
-		const FVector3f CorridorOffset = FDebugDrawing::NavOffset * 1.25f;
-		const FColor Color = ElementToDraw->GetFColor();
-
-		switch (ElementToDraw->GetType())
-		{
-		case EVisualLoggerShapeElement::SinglePoint:
-		{
-			const float Radius = float(ElementToDraw->Radius);
-			const bool bDrawLabel = (ElementToDraw->Description.IsEmpty() == false);
-			const int32 NumPoints = ElementToDraw->Points.Num();
-
-			for (int32 Index = 0; Index < NumPoints; ++Index)
-			{
-				const FVector& Point = ElementToDraw->Points[Index];
-				DebugShapes.Points.Add(FDebugRenderSceneProxy::FSphere(Radius, Point, Color));
-				if (bDrawLabel)
-				{
-					const FString PrintString = NumPoints == 1 ? ElementToDraw->Description : FString::Printf(TEXT("%s_%d"), *ElementToDraw->Description, Index);
-					DebugShapes.Texts.Add(FDebugRenderSceneProxy::FText3d(PrintString, Point, Color));
-				}
-			}
-		}
-			break;
-		case EVisualLoggerShapeElement::Polygon:
-		{
-			FDebugRenderSceneProxy::FMesh TestMesh;
-			GetPolygonMesh(ElementToDraw, TestMesh, CorridorOffset);
-			DebugShapes.Meshes.Add(TestMesh);
-
-			for (int32 VIdx = 0; VIdx < ElementToDraw->Points.Num(); VIdx++)
-			{
-				DebugShapes.Lines.Add(FDebugRenderSceneProxy::FDebugLine(
-					ElementToDraw->Points[VIdx] + CorridorOffset,
-					ElementToDraw->Points[(VIdx + 1) % ElementToDraw->Points.Num()] + CorridorOffset,
-					FColor::Cyan,
-					2)
-					);
-			}
-		}
-			break;
-		case EVisualLoggerShapeElement::Mesh:
-		{
-			struct FHeaderData
-			{
-				float VerticesNum, FacesNum;
-				FHeaderData(const FVector& InVector) : VerticesNum(InVector.X), FacesNum(InVector.Y) {}
-			};
-			const FHeaderData HeaderData(ElementToDraw->Points[0]);
-
-			FDebugRenderSceneProxy::FMesh TestMesh;
-			TestMesh.Color = ElementToDraw->GetFColor();
-			int32 StartIndex = 1;
-			int32 EndIndex = StartIndex + HeaderData.VerticesNum;
-			for (int32 VIdx = StartIndex; VIdx < EndIndex; VIdx++)
-			{
-				TestMesh.Vertices.Add(FVector3f(ElementToDraw->Points[VIdx]));
-			}
-
-
-			StartIndex = EndIndex;
-			EndIndex = StartIndex + HeaderData.FacesNum;
-			for (int32 VIdx = StartIndex; VIdx < EndIndex; VIdx++)
-			{
-				const FVector &CurrentFace = ElementToDraw->Points[VIdx];
-				TestMesh.Indices.Add(CurrentFace.X);
-				TestMesh.Indices.Add(CurrentFace.Y);
-				TestMesh.Indices.Add(CurrentFace.Z);
-			}
-			DebugShapes.Meshes.Add(TestMesh);
-		}
-			break;
-		case EVisualLoggerShapeElement::Segment:
-		{
-			const float Thickness = float(ElementToDraw->Thicknes);
-			const bool bDrawLabel = (ElementToDraw->Description.IsEmpty() == false);
-			const FVector* Location = ElementToDraw->Points.GetData();
-			const int32 NumPoints = ElementToDraw->Points.Num();
-
-			for (int32 Index = 0; Index + 1 < NumPoints; Index += 2, Location += 2)
-			{
-				DebugShapes.Lines.Add(FDebugRenderSceneProxy::FDebugLine(*Location, *(Location + 1), Color, Thickness));
-
-				if (bDrawLabel)
-				{
-					const FString PrintString = NumPoints == 2 ? ElementToDraw->Description : FString::Printf(TEXT("%s_%d"), *ElementToDraw->Description, Index / 2);
-					DebugShapes.Texts.Add(FDebugRenderSceneProxy::FText3d(PrintString, (*Location + (*(Location + 1) - *Location) / 2), Color));
-				}
-			}
-		}
-			break;
-		case EVisualLoggerShapeElement::Path:
-		{
-			const float Thickness = float(ElementToDraw->Thicknes);
-			FVector Location = ElementToDraw->Points[0];
-			for (int32 Index = 1; Index < ElementToDraw->Points.Num(); ++Index)
-			{
-				const FVector CurrentLocation = ElementToDraw->Points[Index];
-				DebugShapes.Lines.Add(FDebugRenderSceneProxy::FDebugLine(Location, CurrentLocation, Color, Thickness));
-				Location = CurrentLocation;
-			}
-		}
-			break;
-		case EVisualLoggerShapeElement::Box:
-		{
-			const float Thickness = float(ElementToDraw->Thicknes);
-			const bool bDrawLabel = (ElementToDraw->Description.IsEmpty() == false);
-			const FVector* BoxExtent = ElementToDraw->Points.GetData();
-			const int32 NumPoints = ElementToDraw->Points.Num();
-
-			for (int32 Index = 0; Index + 1 < NumPoints; Index += 2, BoxExtent += 2)
-			{
-				const FBox Box = FBox(*BoxExtent, *(BoxExtent + 1));
-				DebugShapes.Boxes.Add(FDebugRenderSceneProxy::FDebugBox(Box, Color, FTransform(ElementToDraw->TransformationMatrix)));
-
-				if (bDrawLabel)
-				{
-					const FString PrintString = NumPoints == 2 ? ElementToDraw->Description : FString::Printf(TEXT("%s_%d"), *ElementToDraw->Description, Index / 2);
-					DebugShapes.Texts.Add(FDebugRenderSceneProxy::FText3d(PrintString, Box.GetCenter(), Color));
-				}
-			}
-		}
-			break;
-		case EVisualLoggerShapeElement::Cone:
-		{
-			const float Thickness = float(ElementToDraw->Thicknes);
-			const bool bDrawLabel = ElementToDraw->Description.IsEmpty() == false;
-			for (int32 Index = 0; Index + 2 < ElementToDraw->Points.Num(); Index += 3)
-			{
-				const FVector Orgin = ElementToDraw->Points[Index];
-				const FVector Direction = ElementToDraw->Points[Index + 1].GetSafeNormal();
-				const FVector Angles = ElementToDraw->Points[Index + 2];
-				const float Length = Angles.X;
-
-				FVector YAxis, ZAxis;
-				Direction.FindBestAxisVectors(YAxis, ZAxis);
-				DebugShapes.Cones.Add(FDebugRenderSceneProxy::FCone(FScaleMatrix(FVector(Length)) * FMatrix(Direction, YAxis, ZAxis, Orgin), Angles.Y, Angles.Z, Color));
-
-				if (bDrawLabel)
-				{
-					DebugShapes.Texts.Add(FDebugRenderSceneProxy::FText3d(ElementToDraw->Description, Orgin, Color));
-				}
-			}
-		}
-			break;
-		case EVisualLoggerShapeElement::Cylinder:
-		{
-			const float Thickness = float(ElementToDraw->Thicknes);
-			const bool bDrawLabel = ElementToDraw->Description.IsEmpty() == false;
-			for (int32 Index = 0; Index + 2 < ElementToDraw->Points.Num(); Index += 3)
-			{
-				const FVector Start = ElementToDraw->Points[Index];
-				const FVector End = ElementToDraw->Points[Index + 1];
-				const FVector OtherData = ElementToDraw->Points[Index + 2];
-				const float HalfHeight = 0.5f * (End - Start).Size();
-				const FVector Center = 0.5f * (Start + End);
-				DebugShapes.Cylinders.Add(FDebugRenderSceneProxy::FWireCylinder(Center, OtherData.X, HalfHeight, Color)); // Base parameter is the center of the cylinder
-				if (bDrawLabel)
-				{
-					DebugShapes.Texts.Add(FDebugRenderSceneProxy::FText3d(ElementToDraw->Description, Center, Color));
-				}
-			}
-		}
-			break;
-		case EVisualLoggerShapeElement::Capsule:
-		{
-			const float Thickness = float(ElementToDraw->Thicknes);
-			const bool bDrawLabel = ElementToDraw->Description.IsEmpty() == false;
-			for (int32 Index = 0; Index + 2 < ElementToDraw->Points.Num(); Index += 3)
-			{
-				const FVector Base = ElementToDraw->Points[Index + 0];
-				const FVector FirstData = ElementToDraw->Points[Index + 1];
-				const FVector SecondData = ElementToDraw->Points[Index + 2];
-				const float HalfHeight = FirstData.X;
-				const float Radius = FirstData.Y;
-				const FQuat Rotation = FQuat(FirstData.Z, SecondData.X, SecondData.Y, SecondData.Z);
-
-				const FMatrix Axes = FQuatRotationTranslationMatrix(Rotation, FVector::ZeroVector);
-				const FVector XAxis = Axes.GetScaledAxis(EAxis::X);
-				const FVector YAxis = Axes.GetScaledAxis(EAxis::Y);
-				const FVector ZAxis = Axes.GetScaledAxis(EAxis::Z);
-
-				DebugShapes.Capsules.Add(FDebugRenderSceneProxy::FCapsule(Base, Radius, XAxis, YAxis, ZAxis, HalfHeight, Color));
-				if (bDrawLabel)
-				{
-					DebugShapes.Texts.Add(FDebugRenderSceneProxy::FText3d(ElementToDraw->Description, Base + HalfHeight * FVector::UpVector, Color));
-				}
-			}
-		}
-			break;
-		case EVisualLoggerShapeElement::NavAreaMesh:
-		{
-			if (ElementToDraw->Points.Num() == 0)
-			{
-				continue;
-			}
-
-			struct FHeaderData
-			{
-				float MinZ, MaxZ;
-				FHeaderData(const FVector& InVector) : MinZ(InVector.X), MaxZ(InVector.Y) {}
-			};
-			const FHeaderData HeaderData(ElementToDraw->Points[0]);
-
-			TArray<FVector> AreaMeshPoints = ElementToDraw->Points;
-			AreaMeshPoints.RemoveAt(0, 1, false);
-			AreaMeshPoints.Add(ElementToDraw->Points[1]);
-			TArray<FVector> Vertices;
-			TNavStatArray<FVector> Faces;
-			int32 CurrentIndex = 0;
-			FDebugRenderSceneProxy::FMesh TestMesh;
-			TestMesh.Color = ElementToDraw->GetFColor();
-
-			for (int32 PointIndex = 0; PointIndex < AreaMeshPoints.Num() - 1; PointIndex++)
-			{
-				FVector Point = AreaMeshPoints[PointIndex];
-				FVector NextPoint = AreaMeshPoints[PointIndex + 1];
-
-				FVector3f P1(Point.X, Point.Y, HeaderData.MinZ);
-				FVector3f P2(Point.X, Point.Y, HeaderData.MaxZ);
-				FVector3f P3(NextPoint.X, NextPoint.Y, HeaderData.MinZ);
-				FVector3f P4(NextPoint.X, NextPoint.Y, HeaderData.MaxZ);
-
-				TestMesh.Vertices.Add(P1); TestMesh.Vertices.Add(P2); TestMesh.Vertices.Add(P3);
-				TestMesh.Indices.Add(CurrentIndex + 0);
-				TestMesh.Indices.Add(CurrentIndex + 1);
-				TestMesh.Indices.Add(CurrentIndex + 2);
-				CurrentIndex += 3;
-				TestMesh.Vertices.Add(P3); TestMesh.Vertices.Add(P2); TestMesh.Vertices.Add(P4);
-				TestMesh.Indices.Add(CurrentIndex + 0);
-				TestMesh.Indices.Add(CurrentIndex + 1);
-				TestMesh.Indices.Add(CurrentIndex + 2);
-				CurrentIndex += 3;
-			}
-			DebugShapes.Meshes.Add(TestMesh);
-
-			{
-				FDebugRenderSceneProxy::FMesh PolygonMesh;
-				FVisualLogShapeElement PolygonToDraw(EVisualLoggerShapeElement::Polygon);
-				PolygonToDraw.SetColor(ElementToDraw->GetFColor());
-				PolygonToDraw.Points.Reserve(AreaMeshPoints.Num());
-				PolygonToDraw.Points = AreaMeshPoints;
-				GetPolygonMesh(&PolygonToDraw, PolygonMesh, FVector3f(0, 0, HeaderData.MaxZ));
-				DebugShapes.Meshes.Add(PolygonMesh);
-			}
-
-			for (int32 VIdx = 0; VIdx < AreaMeshPoints.Num(); VIdx++)
-			{
-				DebugShapes.Lines.Add(FDebugRenderSceneProxy::FDebugLine(
-					AreaMeshPoints[VIdx] + FVector(0, 0, HeaderData.MaxZ),
-					AreaMeshPoints[(VIdx + 1) % AreaMeshPoints.Num()] + FVector(0, 0, HeaderData.MaxZ),
-					ElementToDraw->GetFColor(),
-					2)
-					);
-			}
-
-		}
-			break;
-
-		case EVisualLoggerShapeElement::Arrow:
-		{
-			const bool bDrawLabel = (ElementToDraw->Description.IsEmpty() == false);
-			const FVector* Location = ElementToDraw->Points.GetData();
-			const int32 NumPoints = ElementToDraw->Points.Num();
-
-			for (int32 Index = 0; Index + 1 < NumPoints; Index += 2, Location += 2)
-			{
-				DebugShapes.Arrows.Add(FDebugRenderSceneProxy::FArrowLine(*Location, *(Location + 1), Color));
-
-				if (bDrawLabel)
-				{
-					const FString PrintString = NumPoints == 2 ? ElementToDraw->Description : FString::Printf(TEXT("%s_%d"), *ElementToDraw->Description, Index / 2);
-					DebugShapes.Texts.Add(FDebugRenderSceneProxy::FText3d(PrintString, (*Location + (*(Location + 1) - *Location) / 2), Color));
-				}
-			}
-		}
-			break;
-
-		case EVisualLoggerShapeElement::Circle:
-		{
-			const bool bDrawLabel = (ElementToDraw->Description.IsEmpty() == false);
-			const int32 NumPoints = ElementToDraw->Points.Num();
-
-			for (int32 Index = 0; Index + 2 < NumPoints; Index += 3)
-			{
-				const FVector Center = ElementToDraw->Points[Index + 0];
-				const FVector UpAxis = ElementToDraw->Points[Index + 1];
-				const float Radius = ElementToDraw->Points[Index + 2].X;
-				const float Thickness = float(ElementToDraw->Thicknes);
-
-				const FQuat Rotation = FQuat::FindBetweenNormals(FVector::UpVector, UpAxis);
-				const FVector XAxis = Rotation.RotateVector(FVector::XAxisVector);
-				const FVector YAxis = Rotation.RotateVector(FVector::YAxisVector);
-				
-				static const int32 CircleDivs = 12;
-				FVector PrevPosition = FVector::ZeroVector;
-				for (int32 Div = 0; Div <= CircleDivs; Div++)
-				{
-					const float Angle = (float)Div / (float)CircleDivs * PI * 2.0f;
-					const FVector Position = Center + (FMath::Cos(Angle) * XAxis + FMath::Sin(Angle) * YAxis) * Radius;
-					if (Div > 0)
-					{
-						DebugShapes.Lines.Add(FDebugRenderSceneProxy::FDebugLine(PrevPosition, Position, Color, Thickness));
-					}
-					PrevPosition = Position;
-				}
-
-				if (bDrawLabel)
-				{
-					const FString PrintString = NumPoints == 3 ? ElementToDraw->Description : FString::Printf(TEXT("%s_%d"), *ElementToDraw->Description, Index / 3);
-					DebugShapes.Texts.Add(FDebugRenderSceneProxy::FText3d(PrintString, Center, Color));
-				}
-			}
-		}
-			break;
-			
-		}
-	}
-}
+}
\ No newline at end of file
diff --git a/Engine/Source/Developer/LogVisualizer/Private/VisualLoggerRenderingActor.h b/Engine/Source/Developer/LogVisualizer/Private/VisualLoggerRenderingActor.h
--- a/Engine/Source/Developer/LogVisualizer/Private/VisualLoggerRenderingActor.h
+++ b/Engine/Source/Developer/LogVisualizer/Private/VisualLoggerRenderingActor.h
@@ -4,11 +4,14 @@
 
 #include "CoreMinimal.h"
 #include "UObject/ObjectMacros.h"
-#include "GameFramework/Actor.h"
+#include "VisualLoggerRenderingActorBase.h"
 #include "VisualLogger/VisualLoggerTypes.h"
 #include "DebugRenderSceneProxy.h"
 #include "VisualLoggerRenderingActor.generated.h"
 
+// enable adding some hard coded shapes to the VisualLoggerRenderingActor for testing
+#define VLOG_TEST_DEBUG_RENDERING 0
+
 class UPrimitiveComponent;
 struct FVisualLoggerDBRow;
 
@@ -19,36 +22,8 @@ struct FVisualLoggerDBRow;
 DECLARE_MULTICAST_DELEGATE_OneParam(FOnSelectionChanged, class AActor*);
 
 UCLASS(config = Engine, NotBlueprintable, Transient, notplaceable, AdvancedClassDisplay)
-class LOGVISUALIZER_API AVisualLoggerRenderingActor : public AActor
+class LOGVISUALIZER_API AVisualLoggerRenderingActor : public AVisualLoggerRenderingActorBase 
 {
-	struct FTimelineDebugShapes
-	{
-		TArray<FDebugRenderSceneProxy::FDebugLine> Lines;
-		TArray<FDebugRenderSceneProxy::FCone> Cones;
-		TArray<FDebugRenderSceneProxy::FDebugBox> Boxes;
-		TArray<FDebugRenderSceneProxy::FSphere> Points;
-		TArray<FDebugRenderSceneProxy::FMesh> Meshes;
-		TArray<FDebugRenderSceneProxy::FText3d> Texts;
-		TArray<FDebugRenderSceneProxy::FWireCylinder> Cylinders;
-		TArray<FDebugRenderSceneProxy::FCapsule> Capsules;
-		TArray<FDebugRenderSceneProxy::FArrowLine> Arrows;
-		TArray<FVector> LogEntriesPath;
-
-		void Reset()
-		{
-			Lines.Reset();
-			Cones.Reset();
-			Boxes.Reset();
-			Points.Reset();
-			Meshes.Reset();
-			Texts.Reset();
-			Cylinders.Reset();
-			Capsules.Reset();
-			Arrows.Reset();
-			LogEntriesPath.Reset();
-		}
-	};
-
 public:
 	GENERATED_UCLASS_BODY()
 	virtual ~AVisualLoggerRenderingActor();
@@ -56,18 +31,15 @@ public:
 	void ObjectSelectionChanged(const TArray<FName>& Selection);
 	void OnItemSelectionChanged(const FVisualLoggerDBRow& BDRow, int32 ItemIndex);
 
+	virtual void IterateDebugShapes(TFunction<void(const FTimelineDebugShapes&) > Callback) override;
 private:
-	void AddDebugRendering();
-#if ENABLE_VISUAL_LOG
-	void GetDebugShapes(const FVisualLogDevice::FVisualLogEntryItem& EntryItem, FTimelineDebugShapes& DebugShapes);
-#endif
 	void OnFiltersChanged();
 
-public:
-	UPrimitiveComponent* RenderingComponent;
-
-	FTimelineDebugShapes TestDebugShapes;
-
 	TArray<FName> CachedRowSelection;
 	TMap<FName, FTimelineDebugShapes> DebugShapesPerRow;
+
+#if VLOG_TEST_DEBUG_RENDERING
+	void AddDebugRendering();
+	FTimelineDebugShapes TestDebugShapes;
+#endif
 };
diff --git a/Engine/Source/Developer/LogVisualizer/Private/VisualLoggerRenderingActorBase.cpp b/Engine/Source/Developer/LogVisualizer/Private/VisualLoggerRenderingActorBase.cpp
new file mode 100644
--- /dev/null
+++ b/Engine/Source/Developer/LogVisualizer/Private/VisualLoggerRenderingActorBase.cpp
@@ -0,0 +1,550 @@
+// Copyright Epic Games, Inc. All Rights Reserved.
+
+#include "VisualLoggerRenderingActorBase.h"
+#include "AI/NavigationSystemBase.h"
+#include "VisualLogger/VisualLogger.h"
+#include "LogVisualizerSettings.h"
+#include "VisualLoggerDatabase.h"
+#include "LogVisualizerPublic.h"
+#if WITH_EDITOR
+#include "GeomTools.h"
+#endif // WITH_EDITOR
+#include "VisualLoggerRenderingComponent.h"
+
+namespace FDebugDrawing
+{
+	const FVector NavOffset(0, 0, 15);
+}
+
+class UVisualLoggerRenderingComponent;
+class FVisualLoggerSceneProxy final : public FDebugRenderSceneProxy
+{
+public:
+	SIZE_T GetTypeHash() const override
+	{
+		static size_t UniquePointer;
+		return reinterpret_cast<size_t>(&UniquePointer);
+	}
+
+	FVisualLoggerSceneProxy(const UVisualLoggerRenderingComponent* InComponent)
+		: FDebugRenderSceneProxy(InComponent)
+	{
+		DrawType = SolidAndWireMeshes;
+		ViewFlagName = TEXT("VisLog");
+		ViewFlagIndex = uint32(FEngineShowFlags::FindIndexByName(*ViewFlagName));
+		bWantsSelectionOutline = false;
+	}
+
+	virtual FPrimitiveViewRelevance GetViewRelevance(const FSceneView* View) const override
+	{
+		FPrimitiveViewRelevance Result;
+		Result.bDrawRelevance = IsShown(View);
+		Result.bDynamicRelevance = true;
+		// ideally the TranslucencyRelevance should be filled out by the material, here we do it conservative
+		Result.bSeparateTranslucency = Result.bNormalTranslucency = IsShown(View) && GIsEditor;
+		return Result;
+	}
+
+	virtual uint32 GetMemoryFootprint(void) const override { return sizeof(*this) + GetAllocatedSize(); }
+
+	uint32 GetAllocatedSize(void) const
+	{
+		return FDebugRenderSceneProxy::GetAllocatedSize();
+	}
+};
+
+UVisualLoggerRenderingComponent::UVisualLoggerRenderingComponent(const FObjectInitializer& ObjectInitializer) : Super(ObjectInitializer)
+{
+}
+
+FPrimitiveSceneProxy* UVisualLoggerRenderingComponent::CreateSceneProxy()
+{
+	AVisualLoggerRenderingActorBase* RenderingActor = Cast<AVisualLoggerRenderingActorBase>(GetOuter());
+	if (RenderingActor == NULL)
+	{
+		return NULL;
+	}
+
+	ULogVisualizerSettings *Settings = ULogVisualizerSettings::StaticClass()->GetDefaultObject<ULogVisualizerSettings>();
+	FVisualLoggerSceneProxy *VLogSceneProxy = new FVisualLoggerSceneProxy(this);
+	VLogSceneProxy->SolidMeshMaterial = Settings->GetDebugMeshMaterial();
+
+	RenderingActor->IterateDebugShapes([VLogSceneProxy] (const AVisualLoggerRenderingActorBase::FTimelineDebugShapes& Shapes)
+		{
+			VLogSceneProxy->Spheres.Append(Shapes.Points);
+			VLogSceneProxy->Lines.Append(Shapes.Lines);
+			VLogSceneProxy->Boxes.Append(Shapes.Boxes);
+			VLogSceneProxy->Meshes.Append(Shapes.Meshes);
+			VLogSceneProxy->Cones.Append(Shapes.Cones);
+			VLogSceneProxy->Texts.Append(Shapes.Texts);
+			VLogSceneProxy->Cylinders.Append(Shapes.Cylinders);
+			VLogSceneProxy->ArrowLines.Append(Shapes.Arrows);
+			VLogSceneProxy->Capsules.Append(Shapes.Capsules);
+		});
+
+#if WITH_EDITOR
+	if (VLogSceneProxy)
+	{
+		DebugDrawDelegateHelper.InitDelegateHelper(VLogSceneProxy);
+		DebugDrawDelegateHelper.ReregisterDebugDrawDelegate();
+	}
+#endif
+	return VLogSceneProxy;
+}
+
+FBoxSphereBounds UVisualLoggerRenderingComponent::CalcBounds(const FTransform& LocalToWorld) const
+{
+	FBox MyBounds;
+	MyBounds.Init();
+
+	MyBounds = FBox(FVector(-HALF_WORLD_MAX, -HALF_WORLD_MAX, -HALF_WORLD_MAX), FVector(HALF_WORLD_MAX, HALF_WORLD_MAX, HALF_WORLD_MAX));
+
+	return MyBounds;
+}
+
+void UVisualLoggerRenderingComponent::CreateRenderState_Concurrent(FRegisterComponentContext* Context)
+{
+	Super::CreateRenderState_Concurrent(Context);
+
+#if WITH_EDITOR
+	DebugDrawDelegateHelper.RegisterDebugDrawDelegate();
+#endif
+}
+
+void UVisualLoggerRenderingComponent::DestroyRenderState_Concurrent()
+{
+#if WITH_EDITOR
+	DebugDrawDelegateHelper.UnregisterDebugDrawDelegate();
+#endif
+
+	Super::DestroyRenderState_Concurrent();
+}
+
+AVisualLoggerRenderingActorBase::AVisualLoggerRenderingActorBase(const FObjectInitializer& ObjectInitializer)
+	: Super(ObjectInitializer)
+{
+	USceneComponent* SceneComponent = CreateDefaultSubobject<USceneComponent>(TEXT("SceneComp"));
+	RootComponent = SceneComponent;
+
+	RenderingComponent = CreateDefaultSubobject<UVisualLoggerRenderingComponent>(TEXT("RenderingComponent"));
+}
+
+AVisualLoggerRenderingActorBase::~AVisualLoggerRenderingActorBase()
+{
+}
+
+namespace
+{
+	static bool IsPolygonWindingCorrect(const TArray<FVector>& Verts)
+	{
+		// this will work only for convex polys, but we're assuming that all logged polygons are convex in the first place
+		if (Verts.Num() >= 3)
+		{
+			const FVector SurfaceNormal = FVector::CrossProduct(Verts[1] - Verts[0], Verts[2] - Verts[0]);
+			const float TestDot = FVector::DotProduct(SurfaceNormal, FVector(0, 0, 1));
+			return TestDot > 0;
+		}
+
+		return false;
+	}
+
+	static void GetPolygonMesh(const FVisualLogShapeElement* ElementToDraw, FDebugRenderSceneProxy::FMesh& TestMesh, const FVector3f& VertexOffset = FVector3f::ZeroVector)
+	{
+		TestMesh.Color = ElementToDraw->GetFColor();
+
+		FClipSMPolygon InPoly(ElementToDraw->Points.Num());
+		InPoly.FaceNormal = FVector(0, 0, 1);
+
+		const bool bHasCorrectWinding = IsPolygonWindingCorrect(ElementToDraw->Points);
+		if (bHasCorrectWinding)
+		{
+			for (int32 Index = 0; Index < ElementToDraw->Points.Num(); Index++)
+			{
+				FClipSMVertex v1;
+				v1.Pos = ElementToDraw->Points[Index];
+				InPoly.Vertices.Add(v1);
+			}
+		}
+		else
+		{
+			for (int32 Index = ElementToDraw->Points.Num() - 1; Index >= 0; Index--)
+			{
+				FClipSMVertex v1;
+				v1.Pos = ElementToDraw->Points[Index];
+				InPoly.Vertices.Add(v1);
+			}
+		}
+
+		TArray<FClipSMTriangle> OutTris;
+		
+		const bool bTriangulated = FGeomTools::TriangulatePoly(OutTris, InPoly, false);
+		if (bTriangulated)
+		{
+			int32 LastIndex = 0;
+
+			FGeomTools::RemoveRedundantTriangles(OutTris);
+			for (const auto& CurrentTri : OutTris)
+			{
+				TestMesh.Vertices.Add(FDynamicMeshVertex(CurrentTri.Vertices[0].Pos + VertexOffset));
+				TestMesh.Vertices.Add(FDynamicMeshVertex(CurrentTri.Vertices[1].Pos + VertexOffset));
+				TestMesh.Vertices.Add(FDynamicMeshVertex(CurrentTri.Vertices[2].Pos + VertexOffset));
+				TestMesh.Indices.Add(LastIndex++);
+				TestMesh.Indices.Add(LastIndex++);
+				TestMesh.Indices.Add(LastIndex++);
+			}
+		}
+	}
+}
+
+void AVisualLoggerRenderingActorBase::GetDebugShapes(const FVisualLogEntry& InEntry, bool bAddEntryLocationPointer, AVisualLoggerRenderingActorBase::FTimelineDebugShapes& DebugShapes)
+{
+	const FVisualLogEntry* Entry = &InEntry;
+	const FVisualLogShapeElement* ElementToDraw = Entry->ElementsToDraw.GetData();
+	const int32 ElementsCount = Entry->ElementsToDraw.Num();
+
+	if (bAddEntryLocationPointer)
+	{
+		const float Length = 100;
+		const FVector DirectionNorm = FVector(0, 0, 1).GetSafeNormal();
+		FVector YAxis, ZAxis;
+		DirectionNorm.FindBestAxisVectors(YAxis, ZAxis);
+		DebugShapes.Cones.Add(FDebugRenderSceneProxy::FCone(FScaleMatrix(FVector(Length)) * FMatrix(DirectionNorm, YAxis, ZAxis, Entry->Location), 5, 5, FColor::Red));
+	}
+
+	if (DebugShapes.LogEntriesPath.Num())
+	{
+		FVector Location = DebugShapes.LogEntriesPath[0];
+		for (int32 Index = 1; Index < DebugShapes.LogEntriesPath.Num(); ++Index)
+		{
+			const FVector CurrentLocation = DebugShapes.LogEntriesPath[Index];
+			DebugShapes.Lines.Add(FDebugRenderSceneProxy::FDebugLine(Location, CurrentLocation, FColor(160, 160, 240), 2.0));
+			Location = CurrentLocation;
+		}
+	}
+
+	for (int32 ElementIndex = 0; ElementIndex < ElementsCount; ++ElementIndex, ++ElementToDraw)
+	{
+		if (!FVisualLoggerFilters::Get().MatchCategoryFilters(ElementToDraw->Category.ToString(), ElementToDraw->Verbosity))
+		{
+			continue;
+		}
+
+
+		const FVector3f CorridorOffset = FDebugDrawing::NavOffset * 1.25f;
+		const FColor Color = ElementToDraw->GetFColor();
+
+		switch (ElementToDraw->GetType())
+		{
+		case EVisualLoggerShapeElement::SinglePoint:
+		{
+			const float Radius = float(ElementToDraw->Radius);
+			const bool bDrawLabel = (ElementToDraw->Description.IsEmpty() == false);
+			const int32 NumPoints = ElementToDraw->Points.Num();
+
+			for (int32 Index = 0; Index < NumPoints; ++Index)
+			{
+				const FVector& Point = ElementToDraw->Points[Index];
+				DebugShapes.Points.Add(FDebugRenderSceneProxy::FSphere(Radius, Point, Color));
+				if (bDrawLabel)
+				{
+					const FString PrintString = NumPoints == 1 ? ElementToDraw->Description : FString::Printf(TEXT("%s_%d"), *ElementToDraw->Description, Index);
+					DebugShapes.Texts.Add(FDebugRenderSceneProxy::FText3d(PrintString, Point, Color));
+				}
+			}
+		}
+			break;
+		case EVisualLoggerShapeElement::Polygon:
+		{
+			FDebugRenderSceneProxy::FMesh TestMesh;
+			GetPolygonMesh(ElementToDraw, TestMesh, CorridorOffset);
+			DebugShapes.Meshes.Add(TestMesh);
+
+			for (int32 VIdx = 0; VIdx < ElementToDraw->Points.Num(); VIdx++)
+			{
+				DebugShapes.Lines.Add(FDebugRenderSceneProxy::FDebugLine(
+					ElementToDraw->Points[VIdx] + CorridorOffset,
+					ElementToDraw->Points[(VIdx + 1) % ElementToDraw->Points.Num()] + CorridorOffset,
+					FColor::Cyan,
+					2)
+					);
+			}
+		}
+			break;
+		case EVisualLoggerShapeElement::Mesh:
+		{
+			struct FHeaderData
+			{
+				float VerticesNum, FacesNum;
+				FHeaderData(const FVector& InVector) : VerticesNum(InVector.X), FacesNum(InVector.Y) {}
+			};
+			const FHeaderData HeaderData(ElementToDraw->Points[0]);
+
+			FDebugRenderSceneProxy::FMesh TestMesh;
+			TestMesh.Color = ElementToDraw->GetFColor();
+			int32 StartIndex = 1;
+			int32 EndIndex = StartIndex + HeaderData.VerticesNum;
+			for (int32 VIdx = StartIndex; VIdx < EndIndex; VIdx++)
+			{
+				TestMesh.Vertices.Add(FVector3f(ElementToDraw->Points[VIdx]));
+			}
+
+
+			StartIndex = EndIndex;
+			EndIndex = StartIndex + HeaderData.FacesNum;
+			for (int32 VIdx = StartIndex; VIdx < EndIndex; VIdx++)
+			{
+				const FVector &CurrentFace = ElementToDraw->Points[VIdx];
+				TestMesh.Indices.Add(CurrentFace.X);
+				TestMesh.Indices.Add(CurrentFace.Y);
+				TestMesh.Indices.Add(CurrentFace.Z);
+			}
+			DebugShapes.Meshes.Add(TestMesh);
+		}
+			break;
+		case EVisualLoggerShapeElement::Segment:
+		{
+			const float Thickness = float(ElementToDraw->Thicknes);
+			const bool bDrawLabel = (ElementToDraw->Description.IsEmpty() == false);
+			const FVector* Location = ElementToDraw->Points.GetData();
+			const int32 NumPoints = ElementToDraw->Points.Num();
+
+			for (int32 Index = 0; Index + 1 < NumPoints; Index += 2, Location += 2)
+			{
+				DebugShapes.Lines.Add(FDebugRenderSceneProxy::FDebugLine(*Location, *(Location + 1), Color, Thickness));
+
+				if (bDrawLabel)
+				{
+					const FString PrintString = NumPoints == 2 ? ElementToDraw->Description : FString::Printf(TEXT("%s_%d"), *ElementToDraw->Description, Index / 2);
+					DebugShapes.Texts.Add(FDebugRenderSceneProxy::FText3d(PrintString, (*Location + (*(Location + 1) - *Location) / 2), Color));
+				}
+			}
+		}
+			break;
+		case EVisualLoggerShapeElement::Path:
+		{
+			const float Thickness = float(ElementToDraw->Thicknes);
+			FVector Location = ElementToDraw->Points[0];
+			for (int32 Index = 1; Index < ElementToDraw->Points.Num(); ++Index)
+			{
+				const FVector CurrentLocation = ElementToDraw->Points[Index];
+				DebugShapes.Lines.Add(FDebugRenderSceneProxy::FDebugLine(Location, CurrentLocation, Color, Thickness));
+				Location = CurrentLocation;
+			}
+		}
+			break;
+		case EVisualLoggerShapeElement::Box:
+		{
+			const float Thickness = float(ElementToDraw->Thicknes);
+			const bool bDrawLabel = (ElementToDraw->Description.IsEmpty() == false);
+			const FVector* BoxExtent = ElementToDraw->Points.GetData();
+			const int32 NumPoints = ElementToDraw->Points.Num();
+
+			for (int32 Index = 0; Index + 1 < NumPoints; Index += 2, BoxExtent += 2)
+			{
+				const FBox Box = FBox(*BoxExtent, *(BoxExtent + 1));
+				DebugShapes.Boxes.Add(FDebugRenderSceneProxy::FDebugBox(Box, Color, FTransform(ElementToDraw->TransformationMatrix)));
+
+				if (bDrawLabel)
+				{
+					const FString PrintString = NumPoints == 2 ? ElementToDraw->Description : FString::Printf(TEXT("%s_%d"), *ElementToDraw->Description, Index / 2);
+					DebugShapes.Texts.Add(FDebugRenderSceneProxy::FText3d(PrintString, Box.GetCenter(), Color));
+				}
+			}
+		}
+			break;
+		case EVisualLoggerShapeElement::Cone:
+		{
+			const float Thickness = float(ElementToDraw->Thicknes);
+			const bool bDrawLabel = ElementToDraw->Description.IsEmpty() == false;
+			for (int32 Index = 0; Index + 2 < ElementToDraw->Points.Num(); Index += 3)
+			{
+				const FVector Orgin = ElementToDraw->Points[Index];
+				const FVector Direction = ElementToDraw->Points[Index + 1].GetSafeNormal();
+				const FVector Angles = ElementToDraw->Points[Index + 2];
+				const float Length = Angles.X;
+
+				FVector YAxis, ZAxis;
+				Direction.FindBestAxisVectors(YAxis, ZAxis);
+				DebugShapes.Cones.Add(FDebugRenderSceneProxy::FCone(FScaleMatrix(FVector(Length)) * FMatrix(Direction, YAxis, ZAxis, Orgin), Angles.Y, Angles.Z, Color));
+
+				if (bDrawLabel)
+				{
+					DebugShapes.Texts.Add(FDebugRenderSceneProxy::FText3d(ElementToDraw->Description, Orgin, Color));
+				}
+			}
+		}
+			break;
+		case EVisualLoggerShapeElement::Cylinder:
+		{
+			const float Thickness = float(ElementToDraw->Thicknes);
+			const bool bDrawLabel = ElementToDraw->Description.IsEmpty() == false;
+			for (int32 Index = 0; Index + 2 < ElementToDraw->Points.Num(); Index += 3)
+			{
+				const FVector Start = ElementToDraw->Points[Index];
+				const FVector End = ElementToDraw->Points[Index + 1];
+				const FVector OtherData = ElementToDraw->Points[Index + 2];
+				const float HalfHeight = 0.5f * (End - Start).Size();
+				const FVector Center = 0.5f * (Start + End);
+				DebugShapes.Cylinders.Add(FDebugRenderSceneProxy::FWireCylinder(Center, OtherData.X, HalfHeight, Color)); // Base parameter is the center of the cylinder
+				if (bDrawLabel)
+				{
+					DebugShapes.Texts.Add(FDebugRenderSceneProxy::FText3d(ElementToDraw->Description, Center, Color));
+				}
+			}
+		}
+			break;
+		case EVisualLoggerShapeElement::Capsule:
+		{
+			const float Thickness = float(ElementToDraw->Thicknes);
+			const bool bDrawLabel = ElementToDraw->Description.IsEmpty() == false;
+			for (int32 Index = 0; Index + 2 < ElementToDraw->Points.Num(); Index += 3)
+			{
+				const FVector Base = ElementToDraw->Points[Index + 0];
+				const FVector FirstData = ElementToDraw->Points[Index + 1];
+				const FVector SecondData = ElementToDraw->Points[Index + 2];
+				const float HalfHeight = FirstData.X;
+				const float Radius = FirstData.Y;
+				const FQuat Rotation = FQuat(FirstData.Z, SecondData.X, SecondData.Y, SecondData.Z);
+
+				const FMatrix Axes = FQuatRotationTranslationMatrix(Rotation, FVector::ZeroVector);
+				const FVector XAxis = Axes.GetScaledAxis(EAxis::X);
+				const FVector YAxis = Axes.GetScaledAxis(EAxis::Y);
+				const FVector ZAxis = Axes.GetScaledAxis(EAxis::Z);
+
+				DebugShapes.Capsules.Add(FDebugRenderSceneProxy::FCapsule(Base, Radius, XAxis, YAxis, ZAxis, HalfHeight, Color));
+				if (bDrawLabel)
+				{
+					DebugShapes.Texts.Add(FDebugRenderSceneProxy::FText3d(ElementToDraw->Description, Base + HalfHeight * FVector::UpVector, Color));
+				}
+			}
+		}
+			break;
+		case EVisualLoggerShapeElement::NavAreaMesh:
+		{
+			if (ElementToDraw->Points.Num() == 0)
+			{
+				continue;
+			}
+
+			struct FHeaderData
+			{
+				float MinZ, MaxZ;
+				FHeaderData(const FVector& InVector) : MinZ(InVector.X), MaxZ(InVector.Y) {}
+			};
+			const FHeaderData HeaderData(ElementToDraw->Points[0]);
+
+			TArray<FVector> AreaMeshPoints = ElementToDraw->Points;
+			AreaMeshPoints.RemoveAt(0, 1, false);
+			AreaMeshPoints.Add(ElementToDraw->Points[1]);
+			TArray<FVector> Vertices;
+			TNavStatArray<FVector> Faces;
+			int32 CurrentIndex = 0;
+			FDebugRenderSceneProxy::FMesh TestMesh;
+			TestMesh.Color = ElementToDraw->GetFColor();
+
+			for (int32 PointIndex = 0; PointIndex < AreaMeshPoints.Num() - 1; PointIndex++)
+			{
+				FVector Point = AreaMeshPoints[PointIndex];
+				FVector NextPoint = AreaMeshPoints[PointIndex + 1];
+
+				FVector3f P1(Point.X, Point.Y, HeaderData.MinZ);
+				FVector3f P2(Point.X, Point.Y, HeaderData.MaxZ);
+				FVector3f P3(NextPoint.X, NextPoint.Y, HeaderData.MinZ);
+				FVector3f P4(NextPoint.X, NextPoint.Y, HeaderData.MaxZ);
+
+				TestMesh.Vertices.Add(P1); TestMesh.Vertices.Add(P2); TestMesh.Vertices.Add(P3);
+				TestMesh.Indices.Add(CurrentIndex + 0);
+				TestMesh.Indices.Add(CurrentIndex + 1);
+				TestMesh.Indices.Add(CurrentIndex + 2);
+				CurrentIndex += 3;
+				TestMesh.Vertices.Add(P3); TestMesh.Vertices.Add(P2); TestMesh.Vertices.Add(P4);
+				TestMesh.Indices.Add(CurrentIndex + 0);
+				TestMesh.Indices.Add(CurrentIndex + 1);
+				TestMesh.Indices.Add(CurrentIndex + 2);
+				CurrentIndex += 3;
+			}
+			DebugShapes.Meshes.Add(TestMesh);
+
+			{
+				FDebugRenderSceneProxy::FMesh PolygonMesh;
+				FVisualLogShapeElement PolygonToDraw(EVisualLoggerShapeElement::Polygon);
+				PolygonToDraw.SetColor(ElementToDraw->GetFColor());
+				PolygonToDraw.Points.Reserve(AreaMeshPoints.Num());
+				PolygonToDraw.Points = AreaMeshPoints;
+				GetPolygonMesh(&PolygonToDraw, PolygonMesh, FVector3f(0, 0, HeaderData.MaxZ));
+				DebugShapes.Meshes.Add(PolygonMesh);
+			}
+
+			for (int32 VIdx = 0; VIdx < AreaMeshPoints.Num(); VIdx++)
+			{
+				DebugShapes.Lines.Add(FDebugRenderSceneProxy::FDebugLine(
+					AreaMeshPoints[VIdx] + FVector(0, 0, HeaderData.MaxZ),
+					AreaMeshPoints[(VIdx + 1) % AreaMeshPoints.Num()] + FVector(0, 0, HeaderData.MaxZ),
+					ElementToDraw->GetFColor(),
+					2)
+					);
+			}
+
+		}
+			break;
+
+		case EVisualLoggerShapeElement::Arrow:
+		{
+			const bool bDrawLabel = (ElementToDraw->Description.IsEmpty() == false);
+			const FVector* Location = ElementToDraw->Points.GetData();
+			const int32 NumPoints = ElementToDraw->Points.Num();
+
+			for (int32 Index = 0; Index + 1 < NumPoints; Index += 2, Location += 2)
+			{
+				DebugShapes.Arrows.Add(FDebugRenderSceneProxy::FArrowLine(*Location, *(Location + 1), Color));
+
+				if (bDrawLabel)
+				{
+					const FString PrintString = NumPoints == 2 ? ElementToDraw->Description : FString::Printf(TEXT("%s_%d"), *ElementToDraw->Description, Index / 2);
+					DebugShapes.Texts.Add(FDebugRenderSceneProxy::FText3d(PrintString, (*Location + (*(Location + 1) - *Location) / 2), Color));
+				}
+			}
+		}
+			break;
+
+		case EVisualLoggerShapeElement::Circle:
+		{
+			const bool bDrawLabel = (ElementToDraw->Description.IsEmpty() == false);
+			const int32 NumPoints = ElementToDraw->Points.Num();
+
+			for (int32 Index = 0; Index + 2 < NumPoints; Index += 3)
+			{
+				const FVector Center = ElementToDraw->Points[Index + 0];
+				const FVector UpAxis = ElementToDraw->Points[Index + 1];
+				const float Radius = ElementToDraw->Points[Index + 2].X;
+				const float Thickness = float(ElementToDraw->Thicknes);
+
+				const FQuat Rotation = FQuat::FindBetweenNormals(FVector::UpVector, UpAxis);
+				const FVector XAxis = Rotation.RotateVector(FVector::XAxisVector);
+				const FVector YAxis = Rotation.RotateVector(FVector::YAxisVector);
+				
+				static const int32 CircleDivs = 12;
+				FVector PrevPosition = FVector::ZeroVector;
+				for (int32 Div = 0; Div <= CircleDivs; Div++)
+				{
+					const float Angle = (float)Div / (float)CircleDivs * PI * 2.0f;
+					const FVector Position = Center + (FMath::Cos(Angle) * XAxis + FMath::Sin(Angle) * YAxis) * Radius;
+					if (Div > 0)
+					{
+						DebugShapes.Lines.Add(FDebugRenderSceneProxy::FDebugLine(PrevPosition, Position, Color, Thickness));
+					}
+					PrevPosition = Position;
+				}
+
+				if (bDrawLabel)
+				{
+					const FString PrintString = NumPoints == 3 ? ElementToDraw->Description : FString::Printf(TEXT("%s_%d"), *ElementToDraw->Description, Index / 3);
+					DebugShapes.Texts.Add(FDebugRenderSceneProxy::FText3d(PrintString, Center, Color));
+				}
+			}
+		}
+			break;
+			
+		}
+	}
+}
diff --git a/Engine/Source/Developer/LogVisualizer/Public/VisualLoggerRenderingActorBase.h b/Engine/Source/Developer/LogVisualizer/Public/VisualLoggerRenderingActorBase.h
new file mode 100644
--- /dev/null
+++ b/Engine/Source/Developer/LogVisualizer/Public/VisualLoggerRenderingActorBase.h
@@ -0,0 +1,61 @@
+// Copyright Epic Games, Inc. All Rights Reserved.
+
+#pragma once
+
+#include "CoreMinimal.h"
+#include "UObject/ObjectMacros.h"
+#include "GameFramework/Actor.h"
+#include "VisualLogger/VisualLoggerTypes.h"
+#include "DebugRenderSceneProxy.h"
+#include "VisualLoggerRenderingActorBase.generated.h"
+
+/**
+*	Transient actor used to draw visual logger data on level
+*	Base class shared between RewindDebugger and Visual Logger implementations
+*/
+
+UCLASS(config = Engine, NotBlueprintable, Transient, notplaceable, AdvancedClassDisplay)
+class LOGVISUALIZER_API AVisualLoggerRenderingActorBase : public AActor
+{
+public:
+	GENERATED_UCLASS_BODY()
+	virtual ~AVisualLoggerRenderingActorBase();
+
+	struct FTimelineDebugShapes
+	{
+		TArray<FDebugRenderSceneProxy::FDebugLine> Lines;
+		TArray<FDebugRenderSceneProxy::FCone> Cones;
+		TArray<FDebugRenderSceneProxy::FDebugBox> Boxes;
+		TArray<FDebugRenderSceneProxy::FSphere> Points;
+		TArray<FDebugRenderSceneProxy::FMesh> Meshes;
+		TArray<FDebugRenderSceneProxy::FText3d> Texts;
+		TArray<FDebugRenderSceneProxy::FWireCylinder> Cylinders;
+		TArray<FDebugRenderSceneProxy::FCapsule> Capsules;
+		TArray<FDebugRenderSceneProxy::FArrowLine> Arrows;
+		TArray<FVector> LogEntriesPath;
+
+		void Reset()
+		{
+			Lines.Reset();
+			Cones.Reset();
+			Boxes.Reset();
+			Points.Reset();
+			Meshes.Reset();
+			Texts.Reset();
+			Cylinders.Reset();
+			Capsules.Reset();
+			Arrows.Reset();
+			LogEntriesPath.Reset();
+		}
+	};
+
+	// Iterate over each active FTimelineDebugShapes, and call callback
+	virtual void IterateDebugShapes(TFunction<void(const FTimelineDebugShapes&) > Callback) {};
+
+protected:
+#if ENABLE_VISUAL_LOG
+	void GetDebugShapes(const FVisualLogEntry& EntryItem, bool bAddEntryLocationPointer, FTimelineDebugShapes& OutDebugShapes);
+#endif
+
+	UPrimitiveComponent* RenderingComponent;
+};
diff --git a/Engine/Source/Runtime/Engine/Classes/VisualLogger/VisualLoggerTraceDevice.h b/Engine/Source/Runtime/Engine/Classes/VisualLogger/VisualLoggerTraceDevice.h
new file mode 100644
--- /dev/null
+++ b/Engine/Source/Runtime/Engine/Classes/VisualLogger/VisualLoggerTraceDevice.h
@@ -0,0 +1,30 @@
+// Copyright Epic Games, Inc. All Rights Reserved.
+#pragma once
+
+#include "CoreMinimal.h"
+#include "EngineDefines.h"
+#include "VisualLogger/VisualLoggerTypes.h"
+
+#if ENABLE_VISUAL_LOG
+
+class FVisualLoggerTraceDevice : public FVisualLogDevice
+{
+public:
+	static FVisualLoggerTraceDevice& Get()
+	{
+		static FVisualLoggerTraceDevice GDevice;
+		return GDevice;
+	}
+
+	FVisualLoggerTraceDevice();
+	virtual void Cleanup(bool bReleaseMemory = false) override;
+	virtual void StartRecordingToFile(float TImeStamp) override;
+	virtual void StopRecordingToFile(float TImeStamp) override;
+	virtual void DiscardRecordingToFile() override;
+	virtual void SetFileName(const FString& InFileName) override;
+	virtual void Serialize(const class UObject* LogOwner, FName OwnerName, FName OwnerClassName, const FVisualLogEntry& LogEntry) override;
+	virtual bool HasFlags(int32 InFlags) const override { return !!(InFlags & (EVisualLoggerDeviceFlags::CanSaveToFile | EVisualLoggerDeviceFlags::StoreLogsLocally)); }
+
+};
+
+#endif
diff --git a/Engine/Source/Runtime/Engine/Private/VisualLogger/VisualLogger.cpp b/Engine/Source/Runtime/Engine/Private/VisualLogger/VisualLogger.cpp
--- a/Engine/Source/Runtime/Engine/Private/VisualLogger/VisualLogger.cpp
+++ b/Engine/Source/Runtime/Engine/Private/VisualLogger/VisualLogger.cpp
@@ -12,6 +12,7 @@
 #include "AI/NavigationSystemBase.h"
 #include "Framework/Docking/TabManager.h"
 #include "VisualLogger/VisualLoggerBinaryFileDevice.h"
+#include "VisualLogger/VisualLoggerTraceDevice.h"
 #include "VisualLogger/VisualLoggerDebugSnapshotInterface.h"
 #include "Engine/Engine.h"
 #include "TimerManager.h"
@@ -358,6 +359,8 @@ void FVisualLogger::NavigationDataDump(const UObject* Object, const FName& Categ
 FVisualLogger::FVisualLogger()
 {
 	bForceUniqueLogNames = true;
+	bIsRecordingToFile = false;
+	bIsRecordingToTrace = false;
 
 	BlockAllCategories(false);
 	AddDevice(&FVisualLoggerBinaryFileDevice::Get());
@@ -375,11 +378,7 @@ void FVisualLogger::Shutdown()
 {
 	SetIsRecording(false);
 	SetIsRecordingToFile(false);
-
-	if (bUseBinaryFileDevice)
-	{
-		RemoveDevice(&FVisualLoggerBinaryFileDevice::Get());
-	}
+	RemoveDevice(&FVisualLoggerBinaryFileDevice::Get());
 }
 
 void FVisualLogger::Cleanup(UWorld* OldWorld, bool bReleaseMemory)
@@ -585,6 +584,29 @@ void FVisualLogger::SetIsRecordingToFile(bool InIsRecording)
 	bIsRecordingToFile = InIsRecording;
 }
 
+void FVisualLogger::SetIsRecordingToTrace(bool InIsRecording)
+{
+	if (!bIsRecording && InIsRecording)
+	{
+		SetIsRecording(true);
+	}
+
+	FVisualLoggerTraceDevice& Device = FVisualLoggerTraceDevice::Get();
+	if (bIsRecordingToTrace && !InIsRecording)
+	{
+		Device.StopRecordingToFile(0.0);
+		RemoveDevice(&Device);
+	}
+	else if (!bIsRecordingToTrace && InIsRecording)
+	{
+		Device.StartRecordingToFile(0.0);
+		AddDevice(&Device);
+	}
+
+	bIsRecordingToTrace = InIsRecording;
+}
+
+
 void FVisualLogger::DiscardRecordingToFile()
 {
 	if (bIsRecordingToFile)
diff --git a/Engine/Source/Runtime/Engine/Private/VisualLogger/VisualLoggerTraceDevice.cpp b/Engine/Source/Runtime/Engine/Private/VisualLogger/VisualLoggerTraceDevice.cpp
new file mode 100644
--- /dev/null
+++ b/Engine/Source/Runtime/Engine/Private/VisualLogger/VisualLoggerTraceDevice.cpp
@@ -0,0 +1,61 @@
+// Copyright Epic Games, Inc. All Rights Reserved.
+#include "VisualLogger/VisualLoggerTraceDevice.h"
+#include "VisualLogger/VisualLogger.h"
+#include "Serialization/BufferArchive.h"
+#include "ObjectTrace.h"
+
+#include "Trace/Trace.inl"
+
+#if ENABLE_VISUAL_LOG
+
+UE_TRACE_CHANNEL_DEFINE(VisualLoggerChannel);
+
+UE_TRACE_EVENT_BEGIN(VisualLogger, VisualLogEntry)
+	UE_TRACE_EVENT_FIELD(uint64, Cycle)
+	UE_TRACE_EVENT_FIELD(uint64, OwnerId)
+	UE_TRACE_EVENT_FIELD(uint8[], LogEntry)
+UE_TRACE_EVENT_END()
+
+FVisualLoggerTraceDevice::FVisualLoggerTraceDevice()
+{
+
+}
+
+void FVisualLoggerTraceDevice::Cleanup(bool bReleaseMemory)
+{
+
+}
+
+void FVisualLoggerTraceDevice::StartRecordingToFile(float TimeStamp)
+{
+	UE::Trace::ToggleChannel(TEXT("VisualLogger"), true); 
+}
+
+void FVisualLoggerTraceDevice::StopRecordingToFile(float TimeStamp)
+{
+	UE::Trace::ToggleChannel(TEXT("VisualLogger"), false); 
+}
+
+void FVisualLoggerTraceDevice::DiscardRecordingToFile()
+{
+}
+
+void FVisualLoggerTraceDevice::SetFileName(const FString& InFileName)
+{
+}
+
+void FVisualLoggerTraceDevice::Serialize(const UObject* LogOwner, FName OwnerName, FName OwnerClassName, const FVisualLogEntry& LogEntry)
+{
+	if (UE_TRACE_CHANNELEXPR_IS_ENABLED(VisualLoggerChannel))
+	{
+		FBufferArchive Archive;
+		Archive.UsingCustomVersion(EVisualLoggerVersion::GUID);
+		Archive << const_cast<FVisualLogEntry&>(LogEntry);
+
+		UE_TRACE_LOG(VisualLogger, VisualLogEntry, VisualLoggerChannel)
+			<< VisualLogEntry.Cycle(FPlatformTime::Cycles64())
+			<< VisualLogEntry.OwnerId(FObjectTrace::GetObjectId(LogOwner))
+			<< VisualLogEntry.LogEntry(Archive.GetData(), Archive.Num());
+	}
+}
+#endif
diff --git a/Engine/Source/Runtime/Engine/Public/VisualLogger/VisualLogger.h b/Engine/Source/Runtime/Engine/Public/VisualLogger/VisualLogger.h
--- a/Engine/Source/Runtime/Engine/Public/VisualLogger/VisualLogger.h
+++ b/Engine/Source/Runtime/Engine/Public/VisualLogger/VisualLogger.h
@@ -491,6 +491,9 @@ public:
 	/** disables recording to file and discards all data without saving it to file */
 	void DiscardRecordingToFile();
 
+	/** Starts visual log collecting and recording to insights traces (for Rewind Debugger)*/
+	void SetIsRecordingToTrace(bool InIsRecording);
+
 	void SetIsRecordingOnServer(bool IsRecording) { bIsRecordingOnServer = IsRecording; }
 	bool IsRecordingOnServer() const { return !!bIsRecordingOnServer; }
 
@@ -585,16 +588,16 @@ protected:
 	// redirect the trafic to
 	FChildToOwnerRedirectionMap ChildToOwnerMap;
 	// if set all categories are blocked from logging
-	int32 bBlockedAllCategories : 1;
+	bool bBlockedAllCategories : 1;
 	// if set we are recording to file
-	int32 bIsRecordingToFile : 1;
+	bool bIsRecordingToFile : 1;
+	// if set we are recording to insights trace 
+	bool bIsRecordingToTrace : 1;
 	// variable set (from cheat manager) when logging is active on server
-	int32 bIsRecordingOnServer : 1;
+	bool bIsRecordingOnServer : 1;
 	// controls how we generate log names. When set to TRUE there's a lower 
 	// chance of name conflict, but it's more expensive
-	int32 bForceUniqueLogNames : 1;
-	/** Specifies if the Binary Device is being used */
-	int32 bUseBinaryFileDevice : 1;
+	bool bForceUniqueLogNames : 1;
 	// start recording time
 	float StartRecordingToFileTime;
 	/** Delegate to set project specific file name for vlogs */
