[VisualLogger] preserve filters when reseting session

diff --git a/Engine/Source/Developer/LogVisualizer/Private/LogVisualizerSettings.cpp b/Engine/Source/Developer/LogVisualizer/Private/LogVisualizerSettings.cpp
--- a/Engine/Source/Developer/LogVisualizer/Private/LogVisualizerSettings.cpp
+++ b/Engine/Source/Developer/LogVisualizer/Private/LogVisualizerSettings.cpp
@@ -282,20 +282,11 @@ void FVisualLoggerFilters::EnableAllCategories()
 	}
 }
 
-
 void FVisualLoggerFilters::Reset()
 {
-	if (ULogVisualizerSettings::StaticClass()->GetDefaultObject<ULogVisualizerSettings>()->bResetDataWithNewSession == false)
-	{
-		for (int32 Index = Categories.Num() - 1; Index >= 0; --Index)
-		{
-			Categories[Index].bIsInUse = false;
-		}
-	}
-	else
+	for (int32 Index = Categories.Num() - 1; Index >= 0; --Index)
 	{
-		FastCategoryFilterMap.Reset();
-		Categories.Reset();
+		Categories[Index].bIsInUse = false;
 	}
 
 	SearchBoxFilter = FString();
diff --git a/Engine/Source/Developer/LogVisualizer/Private/SVisualLogger.cpp b/Engine/Source/Developer/LogVisualizer/Private/SVisualLogger.cpp
--- a/Engine/Source/Developer/LogVisualizer/Private/SVisualLogger.cpp
+++ b/Engine/Source/Developer/LogVisualizer/Private/SVisualLogger.cpp
@@ -713,11 +713,6 @@ void SVisualLogger::ResetData()
 		MainView->ResetData();
 	}
 
-	if (VisualLoggerFilters.IsValid() && ULogVisualizerSettings::StaticClass()->GetDefaultObject<ULogVisualizerSettings>()->bResetDataWithNewSession)
-	{
-		VisualLoggerFilters->ResetData();
-	}
-
 	if (LogsList.IsValid())
 	{
 		LogsList->ResetData();
