[VisualLogger] added resizable columns in log list

diff --git a/Engine/Source/Developer/LogVisualizer/Private/SVisualLoggerLogsList.cpp b/Engine/Source/Developer/LogVisualizer/Private/SVisualLoggerLogsList.cpp
--- a/Engine/Source/Developer/LogVisualizer/Private/SVisualLoggerLogsList.cpp
+++ b/Engine/Source/Developer/LogVisualizer/Private/SVisualLoggerLogsList.cpp
@@ -32,17 +32,99 @@ namespace ELogsSortMode
 	};
 }
 
+//----------------------------------------------------------------------//
+// SVisualLoggerTableRow
+//----------------------------------------------------------------------//
+typedef TSharedPtr<FLogEntryItem, ESPMode::ThreadSafe> FLogEntryItemPtr;
+
+class SVisualLoggerTableRow : public SMultiColumnTableRow<FLogEntryItemPtr>
+{
+public:
+	SLATE_BEGIN_ARGS(SVisualLoggerTableRow) { }
+	SLATE_END_ARGS()
+
+	static const FName ColumnId_CategoryLabel;
+	static const FName ColumnId_VerbosityLabel;
+	static const FName ColumnId_LogLabel;
+
+	void Construct(const FArguments& InArgs, const TSharedRef<STableViewBase>& InOwnerTableView, const FLogEntryItemPtr InEntryItem)
+	{
+		Item = InEntryItem;
+		SMultiColumnTableRow<FLogEntryItemPtr>::Construct(FSuperRowType::FArguments(), InOwnerTableView);
+	}
+
+	virtual TSharedRef<SWidget> GenerateWidgetForColumn(const FName& ColumnName) override
+	{
+		if (ColumnName == ColumnId_CategoryLabel)
+		{
+			return SNew(STextBlock)
+				.ColorAndOpacity(FSlateColor(Item->CategoryColor))
+				.Text(FText::FromString(Item->Category))
+				.HighlightText(this, &SVisualLoggerTableRow::GetFilterText);
+		}
+		else if (ColumnName == ColumnId_VerbosityLabel)
+		{
+			return SNew(STextBlock)
+				.ColorAndOpacity(FSlateColor(Item->Verbosity == ELogVerbosity::Error ? FLinearColor::Red : (Item->Verbosity == ELogVerbosity::Warning ? FLinearColor::Yellow : FLinearColor::Gray)))
+				.Text(FText::FromString(FString(TEXT("(")) + FString(::ToString(Item->Verbosity)) + FString(TEXT(")"))));
+		}
+		else if (ColumnName == ColumnId_LogLabel)
+		{
+			return SNew(STextBlock)
+				.AutoWrapText(true)
+				.ColorAndOpacity(FSlateColor(Item->Verbosity == ELogVerbosity::Error ? FLinearColor::Red : (Item->Verbosity == ELogVerbosity::Warning ? FLinearColor::Yellow : FLinearColor::Gray)))
+				.Text(FText::FromString(Item->Line))
+				.HighlightText(this, &SVisualLoggerTableRow::GetFilterText)
+				.TextStyle(FLogVisualizerStyle::Get(), TEXT("TextLogs.Text"));
+		}
+
+		return SNullWidget::NullWidget;
+	}
+
+	FText GetFilterText() const
+	{
+		static FText NoText;
+		const bool bSearchInsideLogs = ULogVisualizerSettings::StaticClass()->GetDefaultObject<ULogVisualizerSettings>()->bSearchInsideLogs;
+		return bSearchInsideLogs ? FText::FromString(FVisualLoggerFilters::Get().GetSearchString()) : NoText;
+	}
+
+private:
+	FLogEntryItemPtr Item;
+};
+
+const FName SVisualLoggerTableRow::ColumnId_CategoryLabel = FName("Category");
+const FName SVisualLoggerTableRow::ColumnId_VerbosityLabel = FName("Verbosity");
+const FName SVisualLoggerTableRow::ColumnId_LogLabel = FName("Log");
+
+//----------------------------------------------------------------------//
+// SVisualLoggerLogsList
+//----------------------------------------------------------------------//
 void SVisualLoggerLogsList::Construct(const FArguments& InArgs, const TSharedRef<FUICommandList>& InCommandList)
 {
 	ChildSlot
-		[
-			SAssignNew(LogsLinesWidget, SListView<TSharedPtr<FLogEntryItem> >)
-			.ItemHeight(20)
-			.ListItemsSource(&CachedLogEntryLines)
-			.SelectionMode(ESelectionMode::Multi)
-			.OnSelectionChanged(this, &SVisualLoggerLogsList::LogEntryLineSelectionChanged)
-			.OnGenerateRow(this, &SVisualLoggerLogsList::LogEntryLinesGenerateRow)
-		];
+	[
+		SAssignNew(LogsLinesWidget, SListView<TSharedPtr<FLogEntryItem> >)
+		.ItemHeight(20)
+		.ListItemsSource(&CachedLogEntryLines)
+		.SelectionMode(ESelectionMode::Multi)
+		.OnSelectionChanged(this, &SVisualLoggerLogsList::LogEntryLineSelectionChanged)
+		.OnGenerateRow(this, &SVisualLoggerLogsList::LogEntryLinesGenerateRow)
+		.HeaderRow
+		(
+			SNew(SHeaderRow)
+			+ SHeaderRow::Column(SVisualLoggerTableRow::ColumnId_CategoryLabel)
+			.FillWidth(0.1f)
+			.DefaultLabel(LOCTEXT("VLoggerCategoryLabel", "Category"))
+
+			+ SHeaderRow::Column(SVisualLoggerTableRow::ColumnId_VerbosityLabel)
+			.FillWidth(0.1f)
+			.DefaultLabel(LOCTEXT("VLoggerVerbosityLabel", "Verbosity"))
+
+			+ SHeaderRow::Column(SVisualLoggerTableRow::ColumnId_LogLabel)
+			.FillWidth(0.8f)
+			.DefaultLabel(LOCTEXT("VLoggerLogLabel", "Log"))
+		)
+	];
 
 	FVisualLoggerDatabase::Get().GetEvents().OnItemSelectionChanged.AddRaw(this, &SVisualLoggerLogsList::OnItemSelectionChanged);
 	FVisualLoggerDatabase::Get().GetEvents().OnRowSelectionChanged.AddRaw(this, &SVisualLoggerLogsList::ObjectSelectionChanged);
@@ -85,58 +167,7 @@ void SVisualLoggerLogsList::OnFiltersChanged()
 
 TSharedRef<ITableRow> SVisualLoggerLogsList::LogEntryLinesGenerateRow(TSharedPtr<FLogEntryItem> Item, const TSharedRef<STableViewBase>& OwnerTable)
 {
-	if (Item->Category.Len() > 0)
-	{
-		return SNew(STableRow< TSharedPtr<FString> >, OwnerTable)
-			[
-				SNew(SHorizontalBox)
-				+ SHorizontalBox::Slot()
-				.AutoWidth()
-				.Padding(FMargin(5.0f, 0.0f))
-				[
-					SNew(STextBlock)
-					.ColorAndOpacity(FSlateColor(Item->CategoryColor))
-					.Text(FText::FromString(Item->Category))
-					.HighlightText(this, &SVisualLoggerLogsList::GetFilterText)
-				]
-				+ SHorizontalBox::Slot()
-					.AutoWidth()
-					.Padding(FMargin(5.0f, 0.0f))
-					[
-						SNew(STextBlock)
-						.ColorAndOpacity(FSlateColor(Item->Verbosity == ELogVerbosity::Error ? FLinearColor::Red : (Item->Verbosity == ELogVerbosity::Warning ? FLinearColor::Yellow : FLinearColor::Gray)))
-						.Text(FText::FromString(FString(TEXT("(")) + FString(::ToString(Item->Verbosity)) + FString(TEXT(")"))))
-					]
-				+ SHorizontalBox::Slot()
-					.Padding(FMargin(5.0f, 0.0f))
-					[
-						SNew(STextBlock)
-						.AutoWrapText(true)
-						.ColorAndOpacity(FSlateColor(Item->Verbosity == ELogVerbosity::Error ? FLinearColor::Red : (Item->Verbosity == ELogVerbosity::Warning ? FLinearColor::Yellow : FLinearColor::Gray)))
-						.Text(FText::FromString(Item->Line))
-						.HighlightText(this, &SVisualLoggerLogsList::GetFilterText)
-						.TextStyle(FLogVisualizerStyle::Get(), TEXT("TextLogs.Text"))
-					]
-			];
-
-	}
-	else
-	{
-		return SNew(STableRow< TSharedPtr<FString> >, OwnerTable)
-			[
-				SNew(SHorizontalBox)
-				+ SHorizontalBox::Slot()
-				.Padding(FMargin(5.0f, 0.0f))
-				[
-					SNew(STextBlock)
-					/*.AutoWrapText(true)*/
-					.ColorAndOpacity(FSlateColor(FLinearColor::White))
-					.Text(FText::FromString(Item->Line))
-					.HighlightText(this, &SVisualLoggerLogsList::GetFilterText)
-					.Justification(ETextJustify::Center)
-				]
-			];
-	}
+	return SNew(SVisualLoggerTableRow, OwnerTable, Item);
 }
 
 FText SVisualLoggerLogsList::GetFilterText() const
@@ -185,9 +216,9 @@ void SVisualLoggerLogsList::RegenerateLogEntries()
 	CachedLogEntryLines.Reset();
 
 	const TArray<FName> SelectedRows = FVisualLoggerDatabase::Get().GetSelectedRows();
-	for (FName CurrentRow : SelectedRows)
+	for (const FName CurrentRow : SelectedRows)
 	{
-		if (FVisualLoggerDatabase::Get().IsRowVisible(CurrentRow)  == false)
+		if (FVisualLoggerDatabase::Get().IsRowVisible(CurrentRow) == false)
 		{
 			continue;
 		}
