Close all tab in Visual Logger when the main Tab is closed to prevent a crash when a tab was detached before closing it.

diff --git a/Engine/Source/Developer/LogVisualizer/Private/LogVisualizerModule.cpp b/Engine/Source/Developer/LogVisualizer/Private/LogVisualizerModule.cpp
--- a/Engine/Source/Developer/LogVisualizer/Private/LogVisualizerModule.cpp
+++ b/Engine/Source/Developer/LogVisualizer/Private/LogVisualizerModule.cpp
@@ -71,7 +71,7 @@ void FLogVisualizerModule::StartupModule()
 
 void FLogVisualizerModule::ShutdownModule()
 {
-	FGlobalTabmanager::Get()->UnregisterTabSpawner(VisualLoggerTabName);
+	FGlobalTabmanager::Get()->UnregisterNomadTabSpawner(VisualLoggerTabName);
 	FVisualLoggerCommands::Unregister();
 	IModularFeatures::Get().UnregisterModularFeature(VisualLoggerTabName, this);
 	
diff --git a/Engine/Source/Developer/LogVisualizer/Private/SVisualLogger.cpp b/Engine/Source/Developer/LogVisualizer/Private/SVisualLogger.cpp
--- a/Engine/Source/Developer/LogVisualizer/Private/SVisualLogger.cpp
+++ b/Engine/Source/Developer/LogVisualizer/Private/SVisualLogger.cpp
@@ -111,6 +111,9 @@ SVisualLogger::~SVisualLogger()
 #if WITH_EDITOR
 	FEditorDelegates::PostPIEStarted.Remove(PostPIEStartedHandle);
 #endif
+		
+	TabManager->CloseAllAreas();
+	TabManager->UnregisterAllTabSpawners();
 
 	GEngine->OnWorldAdded().RemoveAll(this);
 
