[VisualLogger] fixed typo in tooltip

diff --git a/Engine/Source/Developer/LogVisualizer/Private/SVisualLoggerFilterWidget.cpp b/Engine/Source/Developer/LogVisualizer/Private/SVisualLoggerFilterWidget.cpp
--- a/Engine/Source/Developer/LogVisualizer/Private/SVisualLoggerFilterWidget.cpp
+++ b/Engine/Source/Developer/LogVisualizer/Private/SVisualLoggerFilterWidget.cpp
@@ -103,7 +103,7 @@ FText SVisualLoggerFilterWidget::GetTooltipString() const
 
 TSharedRef<SWidget> SVisualLoggerFilterWidget::GetRightClickMenuContent()
 {
-	FMenuBuilder MenuBuilder(true, NULL);
+	FMenuBuilder MenuBuilder(true, nullptr);
 
 	if (IsEnabled())
 	{
@@ -115,7 +115,7 @@ TSharedRef<SWidget> SVisualLoggerFilterWidget::GetRightClickMenuContent()
 				const FString VerbosityStr = ::ToString((ELogVerbosity::Type)Index);
 				MenuBuilder.AddMenuEntry(
 					FText::Format(LOCTEXT("UseVerbosity", "Use: {0}"), FText::FromString(VerbosityStr)),
-					LOCTEXT("UseVerbosityTooltip", "Applay verbosity to selected flter."),
+					LOCTEXT("UseVerbosityTooltip", "Apply verbosity to selected filter."),
 					FSlateIcon(),
 					FUIAction(FExecuteAction::CreateSP(this, &SVisualLoggerFilterWidget::SetVerbosityFilter, Index))
 					);
@@ -142,12 +142,12 @@ TSharedRef<SWidget> SVisualLoggerFilterWidget::GetRightClickMenuContent()
 }
 
 bool SVisualLoggerFilterWidget::IsEnabled() const 
-{ 
-	FCategoryFilter& CategoryFilter = FVisualLoggerFilters::Get().GetCategoryByName(GetFilterNameAsString());
+{
+	const FCategoryFilter& CategoryFilter = FVisualLoggerFilters::Get().GetCategoryByName(GetFilterNameAsString());
 	return CategoryFilter.Enabled;
 }
 
-void SVisualLoggerFilterWidget::SetEnabled(bool InEnabled)
+void SVisualLoggerFilterWidget::SetEnabled(const bool InEnabled)
 {
 	FCategoryFilter& CategoryFilter = FVisualLoggerFilters::Get().GetCategoryByName(GetFilterNameAsString());
 	if (InEnabled != CategoryFilter.Enabled)
@@ -157,7 +157,7 @@ void SVisualLoggerFilterWidget::SetEnabled(bool InEnabled)
 	}
 }
 
-void SVisualLoggerFilterWidget::FilterToggled(ECheckBoxState NewState)
+void SVisualLoggerFilterWidget::FilterToggled(const ECheckBoxState NewState)
 {
 	SetEnabled(NewState == ECheckBoxState::Checked);
 }
