PR #8661: [VisualLogger] Added filter volume support to the Visual Logger (Contributed by [REDACTED_PERSON])

Modifications/additions to the PR:
+ replaced usage of FEditorDelegates::EndPie by FGameDelegates::EndPlayMap so the game world tear down is completed and we can udpate the new current world in the same way we do on PostPieStarted
+ fixed FVisualLogEntry constructor to properly initialize members when an invalid actor is provided
+ added refresh for in world primitives on row visibility changes (e.g. adding/moving/deleting a filter volume)

diff --git a/Engine/Source/Developer/LogVisualizer/Private/LogVisualizerSettings.cpp b/Engine/Source/Developer/LogVisualizer/Private/LogVisualizerSettings.cpp
--- a/Engine/Source/Developer/LogVisualizer/Private/LogVisualizerSettings.cpp
+++ b/Engine/Source/Developer/LogVisualizer/Private/LogVisualizerSettings.cpp
@@ -16,6 +16,7 @@ ULogVisualizerSettings::ULogVisualizerSettings(const FObjectInitializer& ObjectI
 	TrivialLogsThreshold = 1;
 	DefaultCameraDistance = 150;
 	bSearchInsideLogs = true;
+	bUseFilterVolumes = true;
 	GraphsBackgroundColor = FColor(0, 0, 0, 70);
 	bResetDataWithNewSession = false;
 	bDrawExtremesOnGraphs = false;
diff --git a/Engine/Source/Developer/LogVisualizer/Private/SVisualLogger.cpp b/Engine/Source/Developer/LogVisualizer/Private/SVisualLogger.cpp
--- a/Engine/Source/Developer/LogVisualizer/Private/SVisualLogger.cpp
+++ b/Engine/Source/Developer/LogVisualizer/Private/SVisualLogger.cpp
@@ -34,9 +34,11 @@
 #include "Editor/EditorEngine.h"
 #include "Editor.h"
 #endif
+#include "GameDelegates.h"
 #include "ISettingsModule.h"
 
 #include "VisualLogger/VisualLoggerBinaryFileDevice.h"
+#include "VisualLogger/VisualLoggerFilterVolume.h"
 
 #define LOCTEXT_NAMESPACE "SVisualLogger"
 
@@ -110,6 +112,7 @@ SVisualLogger::~SVisualLogger()
 {
 #if WITH_EDITOR
 	FEditorDelegates::PostPIEStarted.Remove(PostPIEStartedHandle);
+	FGameDelegates::Get().GetEndPlayMapDelegate().RemoveAll(this);
 #endif
 		
 	TabManager->CloseAllAreas();
@@ -122,6 +125,7 @@ SVisualLogger::~SVisualLogger()
 
 #if WITH_EDITOR
 	ULogVisualizerSettings::StaticClass()->GetDefaultObject<ULogVisualizerSettings>()->SavePersistentData();
+	ULogVisualizerSettings::StaticClass()->GetDefaultObject<ULogVisualizerSettings>()->OnSettingChanged().RemoveAll(this);
 #endif
 
 	if (LastUsedWorld.IsValid())
@@ -144,6 +148,10 @@ SVisualLogger::~SVisualLogger()
 	FLogVisualizer::Get().GetEvents().OnLogLineSelectionChanged.Unbind();
 	FLogVisualizer::Get().GetEvents().OnKeyboardEvent.Unbind();
 
+	GEngine->OnLevelActorAdded().RemoveAll(this);
+	GEngine->OnLevelActorDeleted().RemoveAll(this);
+	GEngine->OnActorMoved().RemoveAll(this);
+
 	FVisualLoggerDatabase::Get().Reset();
 }
 
@@ -167,19 +175,41 @@ void SVisualLogger::Construct(const FArguments& InArgs, const TSharedRef<SDockTa
 	FVisualLoggerDatabase::Get().GetEvents().OnNewItem.AddRaw(this, &SVisualLogger::OnNewItemHandler);
 	FVisualLoggerDatabase::Get().GetEvents().OnItemSelectionChanged.AddRaw(this, &SVisualLogger::OnItemsSelectionChanged);
 
-	GEngine->OnWorldAdded().AddRaw(this, &SVisualLogger::OnNewWorld);
+	LastUsedWorld = FVisualLoggerEditorInterface::Get()->GetWorld();
+	CollectFilterVolumes();
+	ProcessFilterVolumes();
 
 #if WITH_EDITOR
 	PostPIEStartedHandle = FEditorDelegates::PostPIEStarted.AddLambda([this](const bool bIsSimulating)
 	{
-		UWorld* World = FLogVisualizer::Get().GetWorld();
+		UWorld* World = FVisualLoggerEditorInterface::Get()->GetWorld();
+		if (World != nullptr && World != LastUsedWorld)
+		{
+			OnNewWorld(World);
+		}
+	});
+
+	// Note that we use EndPlay delegate instead of FEditorDelegates::EndPie since we want 
+	// the teardown of the PlayWorld to be completed so FVisualLoggerEditorInterface will return
+	// us the Editor world (if any) as the new world.
+	FGameDelegates::Get().GetEndPlayMapDelegate().AddLambda([this]()
+	{
+		UWorld* World = FVisualLoggerEditorInterface::Get()->GetWorld();
 		if (World != nullptr && World != LastUsedWorld)
 		{
 			OnNewWorld(World);
 		}
 	});
+
+	ULogVisualizerSettings::StaticClass()->GetDefaultObject<ULogVisualizerSettings>()->OnSettingChanged().AddRaw(this, &SVisualLogger::OnSettingsChanged);
 #endif
 
+	GEngine->OnWorldAdded().AddRaw(this, &SVisualLogger::OnNewWorld);
+
+	GEngine->OnLevelActorAdded().AddRaw(this, &SVisualLogger::OnLevelActorAdded);
+	GEngine->OnLevelActorDeleted().AddRaw(this, &SVisualLogger::OnLevelActorDeleted);
+	GEngine->OnActorMoved().AddRaw(this, &SVisualLogger::OnActorMoved);
+
 	//////////////////////////////////////////////////////////////////////////
 	// Command Action Lists
 	const FVisualLoggerCommands& Commands = FVisualLoggerCommands::Get();
@@ -751,18 +781,88 @@ void SVisualLogger::ResetData()
 	FLogVisualizer::Get().GetTimeSliderController().Get()->GetTimeSliderArgs().OnScrubPositionChanged = FVisualLoggerTimeSliderArgs::FOnScrubPositionChanged::CreateRaw(this, &SVisualLogger::OnScrubPositionChanged);
 }
 
-void SVisualLogger::OnNewWorld(UWorld* NewWorld)
+void SVisualLogger::CollectFilterVolumes()
 {
-	LastUsedWorld = NewWorld;
+	FilterVolumesInLastUsedWorld.Reset();
+	for (TActorIterator<AActor> It(LastUsedWorld.Get(), AVisualLoggerFilterVolume::StaticClass()); It; ++It)
+	{
+		AVisualLoggerFilterVolume* Volume = Cast<AVisualLoggerFilterVolume>(*It);
+		FilterVolumesInLastUsedWorld.Emplace(Volume);
+	}
+}
 
-	AVisualLoggerRenderingActor* HelperActor = Cast<AVisualLoggerRenderingActor>(FVisualLoggerEditorInterface::Get()->GetHelperActor(LastUsedWorld.Get()));
-	if (ensure(HelperActor))
+void SVisualLogger::ProcessFilterVolumes()
+{
+	FilterBoxes.Reset(FilterVolumesInLastUsedWorld.Num());
+	for (TWeakObjectPtr<const AVisualLoggerFilterVolume>& WeakVolume : FilterVolumesInLastUsedWorld)
 	{
-		if (LastUsedWorld.IsValid() == false || ULogVisualizerSettings::StaticClass()->GetDefaultObject<ULogVisualizerSettings>()->bResetDataWithNewSession)
+		if (const AVisualLoggerFilterVolume* Volume = WeakVolume.Get())
 		{
-			ResetData();
+			FilterBoxes.Add(Volume->GetBounds().GetBox());
 		}
+	}
+}
+
+void SVisualLogger::OnLevelActorAdded(AActor* Actor)
+{
+	if (AVisualLoggerFilterVolume* Volume = Cast<AVisualLoggerFilterVolume>(Actor))
+	{
+		FilterVolumesInLastUsedWorld.AddUnique(Volume);
+		ProcessFilterVolumes();
+		OnFiltersChanged();
+	}
+}
 
+void SVisualLogger::OnLevelActorDeleted(AActor* Actor)
+{
+	if (AVisualLoggerFilterVolume* Volume = Cast<AVisualLoggerFilterVolume>(Actor))
+	{
+		FilterVolumesInLastUsedWorld.Remove(Volume);
+		ProcessFilterVolumes();
+		OnFiltersChanged();
+	}
+}
+
+void SVisualLogger::OnActorMoved(AActor* Actor)
+{
+	if (AVisualLoggerFilterVolume* Volume = Cast<AVisualLoggerFilterVolume>(Actor))
+	{
+		ProcessFilterVolumes();
+		OnFiltersChanged();
+	}
+}
+
+void SVisualLogger::OnSettingsChanged(FName PropertyName)
+{
+	if (PropertyName == GET_MEMBER_NAME_CHECKED(ULogVisualizerSettings, bUseFilterVolumes))
+	{
+		// Note that we don't need to collect & process the volumes since their bookkeeping
+		// is not conditional to the flag, only the filtering.
+		OnFiltersChanged();
+	}
+	else if (PropertyName == GET_MEMBER_NAME_CHECKED(ULogVisualizerSettings, bSearchInsideLogs))
+	{
+		OnFiltersChanged();
+	}
+}
+
+void SVisualLogger::OnNewWorld(UWorld* NewWorld)
+{
+	LastUsedWorld = NewWorld;
+
+	// Only reset data when activating a new game world, not when returning to the Editor world (i.e. after PIE)
+	if (!IsValid(NewWorld) ||
+		(NewWorld->IsGameWorld() && ULogVisualizerSettings::StaticClass()->GetDefaultObject<ULogVisualizerSettings>()->bResetDataWithNewSession))
+	{
+		ResetData();
+	}
+
+	CollectFilterVolumes();
+	ProcessFilterVolumes();
+
+	AVisualLoggerRenderingActor* HelperActor = Cast<AVisualLoggerRenderingActor>(FVisualLoggerEditorInterface::Get()->GetHelperActor(LastUsedWorld.Get()));
+	if (ensure(HelperActor != nullptr))
+	{
 		// reset data and simulate row/item selection to recreate rendering proxy with correct data
 		HelperActor->ResetRendering();
 		const TArray<FName>& SelectedRows = FVisualLoggerDatabase::Get().GetSelectedRows();
@@ -943,6 +1043,9 @@ void SVisualLogger::UpdateVisibilityForEntry(const FVisualLoggerDBRow& DBRow, in
 
 	TArray<FVisualLoggerCategoryVerbosityPair> OutCategories;
 	FVisualLoggerHelpers::GetCategories(CurrentEntry.Entry, OutCategories);
+
+	bool bPassingFilter = true;
+
 	bool bHasValidCategories = false;
 	for (FVisualLoggerCategoryVerbosityPair& CategoryPair : OutCategories)
 	{
@@ -952,8 +1055,24 @@ void SVisualLogger::UpdateVisibilityForEntry(const FVisualLoggerDBRow& DBRow, in
 			break;
 		}
 	}
+	bPassingFilter &= bHasValidCategories;
 
-	if (Settings->bSearchInsideLogs && bHasValidCategories && SearchString.Len() > 0)
+	if (bPassingFilter && Settings->bUseFilterVolumes && FilterBoxes.Num() > 0 && CurrentEntry.Entry.bIsLocationValid)
+	{
+		bool bIsInsideAnyFilterVolume = false;
+		for (const FBox& Box : FilterBoxes)
+		{
+			if (Box.IsInside(CurrentEntry.Entry.Location))
+			{
+				bIsInsideAnyFilterVolume = true;
+				break;
+			}
+		}
+
+		bPassingFilter &= bIsInsideAnyFilterVolume;
+	}
+
+	if (bPassingFilter && Settings->bSearchInsideLogs && SearchString.Len() > 0)
 	{
 		bool bMatchSearchString = false;
 		for (const FVisualLogLine& CurrentLine : CurrentEntry.Entry.LogLines)
@@ -976,12 +1095,10 @@ void SVisualLogger::UpdateVisibilityForEntry(const FVisualLoggerDBRow& DBRow, in
 			}
 		}
 
-		FVisualLoggerDatabase::Get().GetRowByName(DBRow.GetOwnerName()).SetItemVisibility(ItemIndex, bMatchSearchString);
-	}
-	else
-	{
-		FVisualLoggerDatabase::Get().GetRowByName(DBRow.GetOwnerName()).SetItemVisibility(ItemIndex, bHasValidCategories);
+		bPassingFilter &= bMatchSearchString;
 	}
+
+	FVisualLoggerDatabase::Get().GetRowByName(DBRow.GetOwnerName()).SetItemVisibility(ItemIndex, bPassingFilter);
 }
 
 void SVisualLogger::OnLogLineSelectionChanged(TSharedPtr<struct FLogEntryItem> SelectedItem, int64 UserData, FName TagName)
diff --git a/Engine/Source/Developer/LogVisualizer/Private/SVisualLogger.h b/Engine/Source/Developer/LogVisualizer/Private/SVisualLogger.h
--- a/Engine/Source/Developer/LogVisualizer/Private/SVisualLogger.h
+++ b/Engine/Source/Developer/LogVisualizer/Private/SVisualLogger.h
@@ -103,9 +103,19 @@ public:
 	void HandleTabManagerPersistLayout(const TSharedRef<FTabManager::FLayout>& LayoutToSave);
 
 	void OnNewWorld(UWorld* NewWorld);
+	
+	void OnLevelActorAdded(AActor* Actor);
+	void OnLevelActorDeleted(AActor* Actor);
+	void OnActorMoved(AActor* Actor);
+
+	void OnSettingsChanged(FName PropertyName);
+
 	void ResetData();
 
 protected:
+	void CollectFilterVolumes();
+	void ProcessFilterVolumes();
+
 	// Holds the list of UI commands.
 	TSharedRef<FUICommandList> CommandList;
 
@@ -134,4 +144,7 @@ protected:
 
 	FDelegateHandle DrawOnCanvasDelegateHandle;
 	TWeakObjectPtr<class UWorld> LastUsedWorld;
+
+	TArray<TWeakObjectPtr<const class AVisualLoggerFilterVolume>> FilterVolumesInLastUsedWorld;
+	TArray<FBox> FilterBoxes;
 };
diff --git a/Engine/Source/Developer/LogVisualizer/Private/VisualLoggerRenderingActor.cpp b/Engine/Source/Developer/LogVisualizer/Private/VisualLoggerRenderingActor.cpp
--- a/Engine/Source/Developer/LogVisualizer/Private/VisualLoggerRenderingActor.cpp
+++ b/Engine/Source/Developer/LogVisualizer/Private/VisualLoggerRenderingActor.cpp
@@ -3,13 +3,8 @@
 #include "VisualLoggerRenderingActor.h"
 #include "AI/NavigationSystemBase.h"
 #include "VisualLogger/VisualLogger.h"
-#include "LogVisualizerSettings.h"
 #include "VisualLoggerDatabase.h"
 #include "LogVisualizerPublic.h"
-#if WITH_EDITOR
-#include "GeomTools.h"
-#endif // WITH_EDITOR
-#include "VisualLoggerRenderingComponent.h"
 
 AVisualLoggerRenderingActor::AVisualLoggerRenderingActor(const FObjectInitializer& ObjectInitializer)
 	: Super(ObjectInitializer)
@@ -18,6 +13,7 @@ AVisualLoggerRenderingActor::AVisualLoggerRenderingActor(const FObjectInitialize
 	{
 		FVisualLoggerDatabase::Get().GetEvents().OnItemSelectionChanged.AddUObject(this, &AVisualLoggerRenderingActor::OnItemSelectionChanged);
 		FVisualLoggerDatabase::Get().GetEvents().OnRowSelectionChanged.AddUObject(this, &AVisualLoggerRenderingActor::ObjectSelectionChanged);
+		FVisualLoggerDatabase::Get().GetEvents().OnRowChangedVisibility.AddUObject(this, &AVisualLoggerRenderingActor::ObjectVisibilityChanged);
 
 		FLogVisualizer::Get().GetEvents().OnFiltersChanged.AddUObject(this, &AVisualLoggerRenderingActor::OnFiltersChanged);
 	}
@@ -33,10 +29,30 @@ AVisualLoggerRenderingActor::~AVisualLoggerRenderingActor()
 	{
 		FVisualLoggerDatabase::Get().GetEvents().OnItemSelectionChanged.RemoveAll(this);
 		FVisualLoggerDatabase::Get().GetEvents().OnRowSelectionChanged.RemoveAll(this);
+		FVisualLoggerDatabase::Get().GetEvents().OnRowChangedVisibility.RemoveAll(this);
 		FLogVisualizer::Get().GetEvents().OnFiltersChanged.RemoveAll(this);
 	}
 }
 
+void AVisualLoggerRenderingActor::ObjectVisibilityChanged(const FName& RowName)
+{
+	if (DebugShapesPerRow.Contains(RowName) == false)
+	{
+		return;
+	}
+
+	FTimelineDebugShapes& ShapesCache = DebugShapesPerRow[RowName];
+	ShapesCache.Reset();
+
+	if (FVisualLoggerDatabase::Get().IsRowVisible(RowName))
+	{
+		const FVisualLoggerDBRow &DBRow = FVisualLoggerDatabase::Get().GetRowByName(RowName);
+		GetDebugShapes(DBRow.GetItems()[DBRow.GetCurrentItemIndex()].Entry, true, ShapesCache);
+	}
+
+	MarkComponentsRenderStateDirty();
+}
+
 void AVisualLoggerRenderingActor::ObjectSelectionChanged(const TArray<FName>& Selection)
 {
 	if (Selection.Num() > 0)
@@ -74,14 +90,14 @@ void AVisualLoggerRenderingActor::ObjectSelectionChanged(const TArray<FName>& Se
 	MarkComponentsRenderStateDirty();
 }
 
-void AVisualLoggerRenderingActor::OnItemSelectionChanged(const FVisualLoggerDBRow& DBRow, int32 ItemIndex)
+void AVisualLoggerRenderingActor::OnItemSelectionChanged(const FVisualLoggerDBRow& DBRow, const int32 ItemIndex)
 {
 	const FName RowName = DBRow.GetOwnerName();
 
 	const TMap<FName, FVisualLogExtensionInterface*>& AllExtensions = FVisualLogger::Get().GetAllExtensions();
 	for (auto& Extension : AllExtensions)
 	{
-		Extension.Value->DrawData(FVisualLoggerEditorInterface::Get(), NULL);
+		Extension.Value->DrawData(FVisualLoggerEditorInterface::Get(), nullptr);
 	}
 
 	if (DebugShapesPerRow.Contains(RowName) == false)
@@ -115,7 +131,7 @@ void AVisualLoggerRenderingActor::OnFiltersChanged()
 	const TMap<FName, FVisualLogExtensionInterface*>& AllExtensions = FVisualLogger::Get().GetAllExtensions();
 	for (auto& Extension : AllExtensions)
 	{
-		Extension.Value->DrawData(FVisualLoggerEditorInterface::Get(), NULL);
+		Extension.Value->DrawData(FVisualLoggerEditorInterface::Get(), nullptr);
 	}
 
 	DebugShapesPerRow.Reset();
@@ -183,7 +199,7 @@ void AVisualLoggerRenderingActor::AddDebugRendering()
 }
 #endif
 
-void AVisualLoggerRenderingActor::IterateDebugShapes(TFunction<void(const AVisualLoggerRenderingActorBase::FTimelineDebugShapes& Shapes)> Callback)
+void AVisualLoggerRenderingActor::IterateDebugShapes(const TFunction<void(const AVisualLoggerRenderingActorBase::FTimelineDebugShapes& Shapes)> Callback)
 {
 	for (auto& CurrentShapes : DebugShapesPerRow)
 	{
diff --git a/Engine/Source/Developer/LogVisualizer/Private/VisualLoggerRenderingActor.h b/Engine/Source/Developer/LogVisualizer/Private/VisualLoggerRenderingActor.h
--- a/Engine/Source/Developer/LogVisualizer/Private/VisualLoggerRenderingActor.h
+++ b/Engine/Source/Developer/LogVisualizer/Private/VisualLoggerRenderingActor.h
@@ -2,11 +2,7 @@
 
 #pragma once
 
-#include "CoreMinimal.h"
-#include "UObject/ObjectMacros.h"
 #include "VisualLoggerRenderingActorBase.h"
-#include "VisualLogger/VisualLoggerTypes.h"
-#include "DebugRenderSceneProxy.h"
 #include "VisualLoggerRenderingActor.generated.h"
 
 // enable adding some hard coded shapes to the VisualLoggerRenderingActor for testing
@@ -28,6 +24,7 @@ public:
 	GENERATED_UCLASS_BODY()
 	virtual ~AVisualLoggerRenderingActor();
 	void ResetRendering();
+	void ObjectVisibilityChanged(const FName& RowName);
 	void ObjectSelectionChanged(const TArray<FName>& Selection);
 	void OnItemSelectionChanged(const FVisualLoggerDBRow& BDRow, int32 ItemIndex);
 
diff --git a/Engine/Source/Developer/LogVisualizer/Public/LogVisualizerSettings.h b/Engine/Source/Developer/LogVisualizer/Public/LogVisualizerSettings.h
--- a/Engine/Source/Developer/LogVisualizer/Public/LogVisualizerSettings.h
+++ b/Engine/Source/Developer/LogVisualizer/Public/LogVisualizerSettings.h
@@ -145,6 +145,10 @@ public:
 	UPROPERTY(EditAnywhere, config, Category = "VisualLogger")
 	bool bSearchInsideLogs;
 
+	/** Whether to only show events occuring within one of visual logger filter volumes currently in the level */
+	UPROPERTY(EditAnywhere, config, Category = "VisualLogger")
+	bool bUseFilterVolumes;
+
 	/** Background color for 2d graphs visualization */
 	UPROPERTY(EditAnywhere, config, Category = "VisualLogger")
 	FColor GraphsBackgroundColor;
diff --git a/Engine/Source/Runtime/Engine/Private/VisualLogger/VisualLogger.cpp b/Engine/Source/Runtime/Engine/Private/VisualLogger/VisualLogger.cpp
--- a/Engine/Source/Runtime/Engine/Private/VisualLogger/VisualLogger.cpp
+++ b/Engine/Source/Runtime/Engine/Private/VisualLogger/VisualLogger.cpp
@@ -225,6 +225,7 @@ FVisualLogEntry* FVisualLogger::GetEntryToWrite(const UObject* Object, const flo
 		if (const AActor* ObjectAsActor = Cast<AActor>(LogOwner))
 		{
 			CurrentEntry->Location = ObjectAsActor->GetActorLocation();
+			CurrentEntry->bIsLocationValid = true;
 		}
 
 		FOwnerToChildrenRedirectionMap& RedirectionMap = GetRedirectionMap(LogOwner);
diff --git a/Engine/Source/Runtime/Engine/Private/VisualLogger/VisualLoggerFilterVolume.cpp b/Engine/Source/Runtime/Engine/Private/VisualLogger/VisualLoggerFilterVolume.cpp
new file mode 100644
--- /dev/null
+++ b/Engine/Source/Runtime/Engine/Private/VisualLogger/VisualLoggerFilterVolume.cpp
@@ -0,0 +1,8 @@
+#include "VisualLogger/VisualLoggerFilterVolume.h"
+
+AVisualLoggerFilterVolume::AVisualLoggerFilterVolume(const FObjectInitializer& ObjectInitializer)
+{
+	PrimaryActorTick.bCanEverTick = false;
+	PrimaryActorTick.bStartWithTickEnabled = false;
+	bIsEditorOnlyActor = true;
+}
\ No newline at end of file
diff --git a/Engine/Source/Runtime/Engine/Private/VisualLogger/VisualLoggerTypes.cpp b/Engine/Source/Runtime/Engine/Private/VisualLogger/VisualLoggerTypes.cpp
--- a/Engine/Source/Runtime/Engine/Private/VisualLogger/VisualLoggerTypes.cpp
+++ b/Engine/Source/Runtime/Engine/Private/VisualLogger/VisualLoggerTypes.cpp
@@ -59,6 +59,7 @@ FVisualLogEntry::FVisualLogEntry(const FVisualLogEntry& Entry)
 {
 	TimeStamp = Entry.TimeStamp;
 	Location = Entry.Location;
+	bIsLocationValid = Entry.bIsLocationValid;
 
 	Events = Entry.Events;
 	LogLines = Entry.LogLines;
@@ -70,27 +71,31 @@ FVisualLogEntry::FVisualLogEntry(const FVisualLogEntry& Entry)
 
 FVisualLogEntry::FVisualLogEntry(const AActor* InActor, TArray<TWeakObjectPtr<UObject> >* Children)
 {
-	if (IsValid(InActor))
+	if (!IsValid(InActor))
 	{
-		TimeStamp = InActor->GetWorld()->TimeSeconds;
-		Location = InActor->GetActorLocation();
-		const IVisualLoggerDebugSnapshotInterface* DebugSnapshotInterface = Cast<const IVisualLoggerDebugSnapshotInterface>(InActor);
-		if (DebugSnapshotInterface)
-		{
-			DebugSnapshotInterface->GrabDebugSnapshot(this);
-		}
-		if (Children != nullptr)
+		Reset();
+	}
+
+	TimeStamp = InActor->GetWorld()->TimeSeconds;
+	Location = InActor->GetActorLocation();
+	bIsLocationValid = true;
+
+	const IVisualLoggerDebugSnapshotInterface* DebugSnapshotInterface = Cast<const IVisualLoggerDebugSnapshotInterface>(InActor);
+	if (DebugSnapshotInterface)
+	{
+		DebugSnapshotInterface->GrabDebugSnapshot(this);
+	}
+	if (Children != nullptr)
+	{
+		TWeakObjectPtr<UObject>* WeakActorPtr = Children->GetData();
+		for (int32 Index = 0; Index < Children->Num(); ++Index, ++WeakActorPtr)
 		{
-			TWeakObjectPtr<UObject>* WeakActorPtr = Children->GetData();
-			for (int32 Index = 0; Index < Children->Num(); ++Index, ++WeakActorPtr)
+			if (WeakActorPtr->IsValid())
 			{
-				if (WeakActorPtr->IsValid())
+				const IVisualLoggerDebugSnapshotInterface* ChildActor = Cast<const IVisualLoggerDebugSnapshotInterface>(WeakActorPtr->Get());
+				if (ChildActor)
 				{
-					const IVisualLoggerDebugSnapshotInterface* ChildActor = Cast<const IVisualLoggerDebugSnapshotInterface>(WeakActorPtr->Get());
-					if (ChildActor)
-					{
-						ChildActor->GrabDebugSnapshot(this);
-					}
+					ChildActor->GrabDebugSnapshot(this);
 				}
 			}
 		}
@@ -101,6 +106,8 @@ FVisualLogEntry::FVisualLogEntry(float InTimeStamp, FVector InLocation, const UO
 {
 	TimeStamp = InTimeStamp;
 	Location = InLocation;
+	bIsLocationValid = true;
+
 	const IVisualLoggerDebugSnapshotInterface* DebugSnapshotInterface = Cast<const IVisualLoggerDebugSnapshotInterface>(Object);
 	if (DebugSnapshotInterface)
 	{
@@ -127,6 +134,7 @@ void FVisualLogEntry::Reset()
 {
 	TimeStamp = -1;
 	Location = FVector::ZeroVector;
+	bIsLocationValid = false;
 	Events.Reset();
 	LogLines.Reset();
 	Status.Reset();
@@ -482,6 +490,10 @@ FArchive& operator<<(FArchive& Ar, FVisualLogEntry& LogEntry)
 	Ar << LogEntry.Events;
 	Ar << LogEntry.ElementsToDraw;
 	Ar << LogEntry.DataBlocks;
+	
+	uint8 bTempIsLocationValid = (LogEntry.bIsLocationValid != 0);
+	Ar.SerializeBits(&bTempIsLocationValid, 1);
+	LogEntry.bIsLocationValid = bTempIsLocationValid != 0;	
 
 	const int32 VLogsVer = Ar.CustomVer(EVisualLoggerVersion::GUID);
 	if (VLogsVer > EVisualLoggerVersion::Initial)
diff --git a/Engine/Source/Runtime/Engine/Public/VisualLogger/VisualLoggerFilterVolume.h b/Engine/Source/Runtime/Engine/Public/VisualLogger/VisualLoggerFilterVolume.h
new file mode 100644
--- /dev/null
+++ b/Engine/Source/Runtime/Engine/Public/VisualLogger/VisualLoggerFilterVolume.h
@@ -0,0 +1,16 @@
+#pragma once
+
+#include "GameFramework/Volume.h"
+#include "VisualLoggerFilterVolume.generated.h"
+
+/**
+ * A volume to be placed in the level while browsing visual logger output.
+ */
+UCLASS()
+class ENGINE_API AVisualLoggerFilterVolume : public AVolume
+{
+	GENERATED_BODY()
+public:
+	AVisualLoggerFilterVolume(const FObjectInitializer& ObjectInitializer);
+	
+};
\ No newline at end of file
diff --git a/Engine/Source/Runtime/Engine/Public/VisualLogger/VisualLoggerTypes.h b/Engine/Source/Runtime/Engine/Public/VisualLogger/VisualLoggerTypes.h
--- a/Engine/Source/Runtime/Engine/Public/VisualLogger/VisualLoggerTypes.h
+++ b/Engine/Source/Runtime/Engine/Public/VisualLogger/VisualLoggerTypes.h
@@ -166,13 +166,14 @@ struct ENGINE_API FVisualLogEntry
 	uint8 bPassedClassAllowList : 1;
 	uint8 bPassedObjectAllowList : 1;	
 	uint8 bIsAllowedToLog : 1;
+	uint8 bIsLocationValid : 1;
 	uint8 bIsInitialized : 1;
 
 	TArray<FVisualLogEvent> Events;
 	TArray<FVisualLogLine> LogLines;
 	TArray<FVisualLogStatusCategory> Status;
 	TArray<FVisualLogShapeElement> ElementsToDraw;
-	TArray<FVisualLogHistogramSample>	HistogramSamples;
+	TArray<FVisualLogHistogramSample> HistogramSamples;
 	TArray<FVisualLogDataBlock>	DataBlocks;
 
 	FVisualLogEntry() { Reset(); }
