Fix visual logger crash on ObjectVisibilityChanged()

diff --git a/Engine/Source/Developer/LogVisualizer/Private/VisualLoggerRenderingActor.cpp b/Engine/Source/Developer/LogVisualizer/Private/VisualLoggerRenderingActor.cpp
--- a/Engine/Source/Developer/LogVisualizer/Private/VisualLoggerRenderingActor.cpp
+++ b/Engine/Source/Developer/LogVisualizer/Private/VisualLoggerRenderingActor.cpp
@@ -47,7 +47,10 @@ void AVisualLoggerRenderingActor::ObjectVisibilityChanged(const FName& RowName)
 	if (FVisualLoggerDatabase::Get().IsRowVisible(RowName))
 	{
 		const FVisualLoggerDBRow &DBRow = FVisualLoggerDatabase::Get().GetRowByName(RowName);
-		GetDebugShapes(DBRow.GetItems()[DBRow.GetCurrentItemIndex()].Entry, true, ShapesCache);
+		if (DBRow.GetCurrentItemIndex() != INDEX_NONE)
+		{
+			GetDebugShapes(DBRow.GetItems()[DBRow.GetCurrentItemIndex()].Entry, true, ShapesCache);
+		}
 	}
 
 	MarkComponentsRenderStateDirty();
