[LogVisualizer] updated color palette to make filter widgets easier to read

diff --git a/Engine/Source/Developer/LogVisualizer/Private/LogVisualizer.cpp b/Engine/Source/Developer/LogVisualizer/Private/LogVisualizer.cpp
--- a/Engine/Source/Developer/LogVisualizer/Private/LogVisualizer.cpp
+++ b/Engine/Source/Developer/LogVisualizer/Private/LogVisualizer.cpp
@@ -2,7 +2,6 @@
 
 #include "LogVisualizer.h"
 #include "LogVisualizerPublic.h"
-#include "EngineGlobals.h"
 #include "GameFramework/SpectatorPawn.h"
 #include "Engine/LocalPlayer.h"
 #include "EngineUtils.h"
@@ -17,51 +16,51 @@
 
 TSharedPtr< struct FLogVisualizer > FLogVisualizer::StaticInstance;
 FColor FLogVisualizerColorPalette[] = {
-	FColor(0xff00A480),
-	FColorList::Aquamarine,
-	FColorList::Cyan,
-	FColorList::Brown,
-	FColorList::Green,
-	FColorList::Orange,
-	FColorList::Magenta,
-	FColorList::BrightGold,
-	FColorList::NeonBlue,
-	FColorList::MediumSlateBlue,
-	FColorList::SpicyPink,
-	FColorList::SpringGreen,
-	FColorList::SteelBlue,
-	FColorList::SummerSky,
-	FColorList::Violet,
-	FColorList::VioletRed,
-	FColorList::YellowGreen,
-	FColor(0xff62E200),
-	FColor(0xff1F7B67),
-	FColor(0xff62AA2A),
-	FColor(0xff70227E),
-	FColor(0xff006B53),
-	FColor(0xff409300),
-	FColor(0xff5D016D),
-	FColor(0xff34D2AF),
-	FColor(0xff8BF13C),
-	FColor(0xffBC38D3),
-	FColor(0xff5ED2B8),
-	FColor(0xffA6F16C),
-	FColor(0xffC262D3),
-	FColor(0xff0F4FA8),
-	FColor(0xff00AE68),
-	FColor(0xffDC0055),
-	FColor(0xff284C7E),
-	FColor(0xff21825B),
-	FColor(0xffA52959),
-	FColor(0xff05316D),
-	FColor(0xff007143),
-	FColor(0xff8F0037),
-	FColor(0xff4380D3),
-	FColor(0xff36D695),
-	FColor(0xffEE3B80),
-	FColor(0xff6996D3),
-	FColor(0xff60D6A7),
-	FColor(0xffEE6B9E)
+	FColor(0xff8fbc8f), // darkseagreen
+	FColor(0xff48d1cc), // mediumturquoise
+	FColor(0xffffa500), // orange
+	FColor(0xff6a5acd), // slateblue
+	FColor(0xffffff00), // yellow
+	FColor(0xffdeb887), // burlywood
+	FColor(0xff00ff00), // lime
+	FColor(0xffdc143c), // crimson
+	FColor(0xff00fa9a), // mediumspringgreen
+	FColor(0xff00bfff), // deepskyblue
+	FColor(0xffadff2f), // greenyellow
+	FColor(0xffff00ff), // fuchsia
+	FColor(0xff1e90ff), // dodgerblue
+	FColor(0xfff0e68c), // khaki
+	FColor(0xffdda0dd), // plum
+	FColor(0xff90ee90), // lightgreen
+	FColor(0xffff1493), // deeppink
+	FColor(0xffffa07a), // lightsalmon
+	FColor(0xffee82ee), // violet
+	FColor(0xff7fffd4), // aquamarine
+	FColor(0xffff0000), // red
+	FColor(0xfffafad2), // lightgoldenrod
+	FColor(0xffcd5c5c), // indianred
+	FColor(0xffe6e6fa), // lavender
+	FColor(0xffffb6c1), // lightpink
+	FColor(0xffa9a9a9), // darkgray
+	FColor(0xff9932cc), // darkorchid
+	FColor(0xff556b2f), // darkolivegreen
+	FColor(0xffb03060), // maroon3
+	FColor(0xff8b4513), // saddlebrown
+	FColor(0xff228b22), // forestgreen
+	FColor(0xff808000), // olive
+	FColor(0xff3cb371), // mediumseagreen
+	FColor(0xffb8860b), // darkgoldenrod
+	FColor(0xff008b8b), // darkcyan
+	FColor(0xff4682b4), // steelblue
+	FColor(0xff32cd32), // limegreen
+	FColor(0xffd2691e), // chocolate
+	FColor(0xff9acd32), // yellowgreen
+	FColor(0xff191970), // midnightblue
+	FColor(0xff8b0000), // darkred
+	FColor(0xff0000ff), // blue
+	FColor(0xff00008b), // darkblue
+	FColor(0xff7f007f), // purple2
+	FColor(0xff2f4f4f), // darkslategray
 };
 
 
