Addressing numeric conversion warnings in visual logger plus minor code tidy up.

diff --git a/Engine/Source/Developer/LogVisualizer/Private/LogVisualizer.cpp b/Engine/Source/Developer/LogVisualizer/Private/LogVisualizer.cpp
--- a/Engine/Source/Developer/LogVisualizer/Private/LogVisualizer.cpp
+++ b/Engine/Source/Developer/LogVisualizer/Private/LogVisualizer.cpp
@@ -190,7 +190,7 @@ void FLogVisualizer::UpdateCameraPosition(FName RowName, int32 ItemIndes)
 
 			FBox BoundingBox = FBox::BuildAABB(CurrentLocation, Extent);
 			const FVector Position = BoundingBox.GetCenter();
-			float Radius = BoundingBox.GetExtent().Size();
+			float Radius = static_cast<float>(BoundingBox.GetExtent().Size());
 			
 			FViewportCameraTransform ViewTransform;
 			ViewTransform.TransitionToLocation(Position, nullptr, true);
@@ -199,10 +199,10 @@ void FLogVisualizer::UpdateCameraPosition(FName RowName, int32 ItemIndes)
 			const float AspectRatio = 1.777777f;
 			CA_SUPPRESS(6326);
 			uint32 MinAxisSize = (AspectRatio > 1.0f) ? Viewport->GetSizeXY().Y : Viewport->GetSizeXY().X;
-			float Zoom = Radius / (MinAxisSize / 2.0f);
+			const float Zoom = Radius / (MinAxisSize / 2.0f);
 
 			NewOrthoZoom = Zoom * (Viewport->GetSizeXY().X*15.0f);
-			NewOrthoZoom = FMath::Clamp<float>(NewOrthoZoom, 250, MAX_FLT);
+			NewOrthoZoom = FMath::Clamp<float>(NewOrthoZoom, 250.0f, MAX_FLT);
 			ViewTransform.SetOrthoZoom(NewOrthoZoom);
 
 			AVisualLoggerCameraController::Instance->GetSpectatorPawn()->TeleportTo(ViewTransform.GetLocation(), ViewTransform.GetRotation(), false, true);
@@ -264,6 +264,7 @@ int32 FLogVisualizer::GetPreviousItem(FName RowName, int32 MoveDistance)
 	return NewItemIndex;
 }
 
+// @todo: This function currently doesn't do anything!
 void FLogVisualizer::GotoNextItem(FName RowName, int32 MoveDistance)
 {
 	FVisualLoggerDBRow& DBRow = FVisualLoggerDatabase::Get().GetRowByName(RowName);
@@ -271,11 +272,12 @@ void FLogVisualizer::GotoNextItem(FName RowName, int32 MoveDistance)
 
 	if (NewItemIndex != DBRow.GetCurrentItemIndex())
 	{
-		auto &Entries = DBRow.GetItems();
+		auto& Entries = DBRow.GetItems();
 		float NewTimeStamp = Entries[NewItemIndex].Entry.TimeStamp;
 	}
 }
 
+// @todo: This function currently doesn't do anything!
 void FLogVisualizer::GotoPreviousItem(FName RowName, int32 MoveDistance)
 {
 	FVisualLoggerDBRow& DBRow = FVisualLoggerDatabase::Get().GetRowByName(RowName);
@@ -283,7 +285,7 @@ void FLogVisualizer::GotoPreviousItem(FName RowName, int32 MoveDistance)
 
 	if (NewItemIndex != DBRow.GetCurrentItemIndex())
 	{
-		auto &Entries = DBRow.GetItems();
+		auto& Entries = DBRow.GetItems();
 		float NewTimeStamp = Entries[NewItemIndex].Entry.TimeStamp;
 	}
 }
diff --git a/Engine/Source/Developer/LogVisualizer/Private/SVisualLogger.cpp b/Engine/Source/Developer/LogVisualizer/Private/SVisualLogger.cpp
--- a/Engine/Source/Developer/LogVisualizer/Private/SVisualLogger.cpp
+++ b/Engine/Source/Developer/LogVisualizer/Private/SVisualLogger.cpp
@@ -877,7 +877,7 @@ void SVisualLogger::OnNewWorld(UWorld* NewWorld)
 
 void SVisualLogger::OnObjectSelectionChanged(const TArray<FName>& RowNames)
 {
-	const float ScrubTime = FLogVisualizer::Get().GetTimeSliderController().Get()->GetTimeSliderArgs().ScrubPosition.Get();
+	const double ScrubTime = FLogVisualizer::Get().GetTimeSliderController().Get()->GetTimeSliderArgs().ScrubPosition.Get();
 	for (auto RowName : RowNames)
 	{
 		FVisualLoggerDBRow& DBRow = FVisualLoggerDatabase::Get().GetRowByName(RowName);
@@ -949,7 +949,7 @@ void SVisualLogger::OnFiltersChanged()
 	const int32 BlockingCycles = int32(FPlatformTime::Cycles() - StartCycles);
 	{
 		const TArray<FName>& SelectedRows = FVisualLoggerDatabase::Get().GetSelectedRows();
-		const float ScrubTime = FLogVisualizer::Get().GetTimeSliderController()->GetTimeSliderArgs().ScrubPosition.Get();
+		const double ScrubTime = FLogVisualizer::Get().GetTimeSliderController()->GetTimeSliderArgs().ScrubPosition.Get();
 		{
 			for (auto RowName : SelectedRows)
 			{
@@ -1114,10 +1114,10 @@ void SVisualLogger::OnLogLineSelectionChanged(TSharedPtr<struct FLogEntryItem> S
 	}
 }
 
-void SVisualLogger::OnScrubPositionChanged(float NewScrubPosition, bool bScrubbing)
+void SVisualLogger::OnScrubPositionChanged(double NewScrubPosition, bool bScrubbing)
 {
 	const TArray<FName> &SelectedRows = FVisualLoggerDatabase::Get().GetSelectedRows();
-	const float ScrubTime = FLogVisualizer::Get().GetTimeSliderController()->GetTimeSliderArgs().ScrubPosition.Get();
+	const double ScrubTime = FLogVisualizer::Get().GetTimeSliderController()->GetTimeSliderArgs().ScrubPosition.Get();
 	for (auto RowName : SelectedRows)
 	{
 		auto& DBRow = FVisualLoggerDatabase::Get().GetRowByName(RowName);
@@ -1154,9 +1154,9 @@ FReply SVisualLogger::OnKeyboardRedirection(const FGeometry& MyGeometry, const F
 	const FKey Key = InKeyEvent.GetKey();
 	if (Key == EKeys::Left || Key == EKeys::Right)
 	{
-		const float ScrubTime = FLogVisualizer::Get().GetTimeSliderController()->GetTimeSliderArgs().ScrubPosition.Get();
-		float NewTimeToSet = ScrubTime;
-		float BestTimeDifference = FLT_MAX;
+		const double ScrubTime = FLogVisualizer::Get().GetTimeSliderController()->GetTimeSliderArgs().ScrubPosition.Get();
+		double NewTimeToSet = ScrubTime;
+		double BestTimeDifference = MAX_dbl;
 
 		const int32 MoveDist = InKeyEvent.IsLeftControlDown() ? InKeyEvent.IsLeftShiftDown() ? 20 : 10 : 1;
 		for (auto RowName : SelectedRows)
@@ -1170,7 +1170,7 @@ FReply SVisualLogger::OnKeyboardRedirection(const FGeometry& MyGeometry, const F
 
 			if (Key == EKeys::Right)
 			{
-				float TimeDifference = DBRow.GetCurrentItem().Entry.TimeStamp - ScrubTime;
+				double TimeDifference = DBRow.GetCurrentItem().Entry.TimeStamp - ScrubTime;
 				if (TimeDifference > 0 && FMath::Abs(TimeDifference) < FMath::Abs(BestTimeDifference))
 				{
 					BestTimeDifference = TimeDifference;
@@ -1187,7 +1187,7 @@ FReply SVisualLogger::OnKeyboardRedirection(const FGeometry& MyGeometry, const F
 			}
 			else if (Key == EKeys::Left)
 			{
-				float TimeDifference = DBRow.GetCurrentItem().Entry.TimeStamp - ScrubTime;
+				double TimeDifference = DBRow.GetCurrentItem().Entry.TimeStamp - ScrubTime;
 				if (TimeDifference < 0 && FMath::Abs(TimeDifference) < FMath::Abs(BestTimeDifference))
 				{
 					BestTimeDifference = TimeDifference;
diff --git a/Engine/Source/Developer/LogVisualizer/Private/SVisualLogger.h b/Engine/Source/Developer/LogVisualizer/Private/SVisualLogger.h
--- a/Engine/Source/Developer/LogVisualizer/Private/SVisualLogger.h
+++ b/Engine/Source/Developer/LogVisualizer/Private/SVisualLogger.h
@@ -59,7 +59,7 @@ public:
 	void OnItemsSelectionChanged(const FVisualLoggerDBRow&, int32);
 	void OnNewItemHandler(const FVisualLoggerDBRow& BDRow, int32 ItemIndex);
 	void UpdateVisibilityForEntry(const FVisualLoggerDBRow& BDRow, int32 ItemIndex);
-	void OnScrubPositionChanged(float NewScrubPosition, bool bScrubbing);
+	void OnScrubPositionChanged(double NewScrubPosition, bool bScrubbing);
 
 	void OnFiltersSearchChanged(const FText& Filter);
 	void OnLogLineSelectionChanged(TSharedPtr<struct FLogEntryItem> SelectedItem, int64 UserData, FName TagName);
diff --git a/Engine/Source/Developer/LogVisualizer/Private/SVisualLoggerFilterWidget.cpp b/Engine/Source/Developer/LogVisualizer/Private/SVisualLoggerFilterWidget.cpp
--- a/Engine/Source/Developer/LogVisualizer/Private/SVisualLoggerFilterWidget.cpp
+++ b/Engine/Source/Developer/LogVisualizer/Private/SVisualLoggerFilterWidget.cpp
@@ -35,7 +35,7 @@ void SVisualLoggerFilterWidget::Construct(const FArguments& InArgs)
 	ChildSlot
 		[
 			SNew(SBorder)
-			.Padding(2)
+			.Padding(2.f)
 			.BorderBackgroundColor(this, &SVisualLoggerFilterWidget::GetBorderBackgroundColor)
 			.BorderImage(FLogVisualizerStyle::Get().GetBrush("ContentBrowser.FilterButtonBorder"))
 			[
diff --git a/Engine/Source/Developer/LogVisualizer/Private/SVisualLoggerFilters.cpp b/Engine/Source/Developer/LogVisualizer/Private/SVisualLoggerFilters.cpp
--- a/Engine/Source/Developer/LogVisualizer/Private/SVisualLoggerFilters.cpp
+++ b/Engine/Source/Developer/LogVisualizer/Private/SVisualLoggerFilters.cpp
@@ -37,7 +37,7 @@ void SVisualLoggerFilters::Construct(const FArguments& InArgs, const TSharedRef<
 		SNew(SComboButton)
 		.ComboButtonStyle(FLogVisualizerStyle::Get(), "Filters.Style")
 		.ForegroundColor(FLinearColor::White)
-		.ContentPadding(0)
+		.ContentPadding(0.f)
 		.OnGetMenuContent(this, &SVisualLoggerFilters::MakeGraphsFilterMenu)
 		.ToolTipText(LOCTEXT("AddFilterToolTip", "Add an asset filter."))
 		.HasDownArrow(true)
diff --git a/Engine/Source/Developer/LogVisualizer/Private/SVisualLoggerLogsList.cpp b/Engine/Source/Developer/LogVisualizer/Private/SVisualLoggerLogsList.cpp
--- a/Engine/Source/Developer/LogVisualizer/Private/SVisualLoggerLogsList.cpp
+++ b/Engine/Source/Developer/LogVisualizer/Private/SVisualLoggerLogsList.cpp
@@ -227,7 +227,7 @@ void SVisualLoggerLogsList::RegenerateLogEntries()
 		const TArray<FVisualLogDevice::FVisualLogEntryItem>& Entries = FVisualLoggerDatabase::Get().GetRowByName(CurrentRow).GetItems();
 		
 		int32 BestItemIndex = INDEX_NONE;
-		float BestDistance = MAX_FLT;
+		double BestDistance = MAX_dbl;
 		for (int32 Index = 0; Index < Entries.Num(); Index++)
 		{
 			auto& CurrentEntryItem = Entries[Index];
@@ -237,7 +237,7 @@ void SVisualLoggerLogsList::RegenerateLogEntries()
 			}
 
 			TArray<FVisualLoggerCategoryVerbosityPair> OutCategories;
-			const float CurrentDist = DBRow.GetCurrentItemIndex() == INDEX_NONE ? 0 : FMath::Abs(CurrentEntryItem.Entry.TimeStamp - DBRow.GetCurrentItem().Entry.TimeStamp);
+			const double CurrentDist = DBRow.GetCurrentItemIndex() == INDEX_NONE ? 0. : FMath::Abs(CurrentEntryItem.Entry.TimeStamp - DBRow.GetCurrentItem().Entry.TimeStamp);
 			if (CurrentDist < BestDistance)
 			{
 				BestDistance = CurrentDist;
diff --git a/Engine/Source/Developer/LogVisualizer/Private/SVisualLoggerReport.cpp b/Engine/Source/Developer/LogVisualizer/Private/SVisualLoggerReport.cpp
--- a/Engine/Source/Developer/LogVisualizer/Private/SVisualLoggerReport.cpp
+++ b/Engine/Source/Developer/LogVisualizer/Private/SVisualLoggerReport.cpp
@@ -51,7 +51,7 @@ void SVisualLoggerReport::Construct(const FArguments& InArgs, TArray< TSharedPtr
 					.Padding(5.0f)
 					.BorderImage(FAppStyle::Get().GetBrush("BoxShadow"))
 					[
-						SNew(SBorder).Padding(2).HAlign(HAlign_Left)
+						SNew(SBorder).Padding(2.f).HAlign(HAlign_Left)
 						.BorderImage(FLogVisualizerStyle::Get().GetBrush("RichText.Background"))
 						[
 							SAssignNew(InteractiveRichText, SRichTextBlock)
diff --git a/Engine/Source/Developer/LogVisualizer/Private/SVisualLoggerStatusView.cpp b/Engine/Source/Developer/LogVisualizer/Private/SVisualLoggerStatusView.cpp
--- a/Engine/Source/Developer/LogVisualizer/Private/SVisualLoggerStatusView.cpp
+++ b/Engine/Source/Developer/LogVisualizer/Private/SVisualLoggerStatusView.cpp
@@ -24,7 +24,7 @@ void SVisualLoggerStatusView::Construct(const FArguments& InArgs, const TSharedR
 	ChildSlot
 		[
 			SNew(SBorder)
-			.Padding(1)
+			.Padding(1.f)
 			.BorderImage(FLogVisualizerStyle::Get().GetBrush("ToolPanel.GroupBorder"))
 			[
 				SNew(SVerticalBox)
diff --git a/Engine/Source/Developer/LogVisualizer/Private/SVisualLoggerTimeline.cpp b/Engine/Source/Developer/LogVisualizer/Private/SVisualLoggerTimeline.cpp
--- a/Engine/Source/Developer/LogVisualizer/Private/SVisualLoggerTimeline.cpp
+++ b/Engine/Source/Developer/LogVisualizer/Private/SVisualLoggerTimeline.cpp
@@ -240,21 +240,25 @@ const FSlateBrush* SLogVisualizerTimeline::GetBorder() const
 	}
 }
 
-void SLogVisualizerTimeline::Goto(float ScrubPosition) 
-{ 
+// @todo: This function currently doesn't do anything!
+void SLogVisualizerTimeline::Goto(float ScrubPosition)
+{
 
 }
 
-void SLogVisualizerTimeline::GotoNextItem() 
-{ 
+// @todo: This function currently doesn't do anything!
+void SLogVisualizerTimeline::GotoNextItem()
+{
 	FLogVisualizer::Get().GotoNextItem(GetName());
 }
 
-void SLogVisualizerTimeline::GotoPreviousItem() 
-{ 
+// @todo: This function currently doesn't do anything!
+void SLogVisualizerTimeline::GotoPreviousItem()
+{
 	FLogVisualizer::Get().GotoPreviousItem(GetName());
 }
 
+// @todo: This function currently doesn't do anything!
 void SLogVisualizerTimeline::MoveCursorByDistance(int32 Distance)
 {
 	if (Distance > 0)
@@ -265,5 +269,4 @@ void SLogVisualizerTimeline::MoveCursorByDistance(int32 Distance)
 	{
 		FLogVisualizer::Get().GotoPreviousItem(GetName(), FMath::Abs(Distance));
 	}
-	
 }
diff --git a/Engine/Source/Developer/LogVisualizer/Private/SVisualLoggerTimeline.h b/Engine/Source/Developer/LogVisualizer/Private/SVisualLoggerTimeline.h
--- a/Engine/Source/Developer/LogVisualizer/Private/SVisualLoggerTimeline.h
+++ b/Engine/Source/Developer/LogVisualizer/Private/SVisualLoggerTimeline.h
@@ -53,9 +53,13 @@ public:
 	void HandleLogVisualizerSettingChanged(FName Name);
 	TSharedPtr<class SVisualLoggerTimelinesContainer> GetOwner() { return Owner; }
 
+	// @todo: This function currently doesn't do anything!
 	void Goto(float ScrubPosition);
+	// @todo: This function currently doesn't do anything!
 	void GotoNextItem();
+	// @todo: This function currently doesn't do anything!
 	void GotoPreviousItem();
+	// @todo: This function currently doesn't do anything!
 	void MoveCursorByDistance(int32 Distance);
 
 protected:
diff --git a/Engine/Source/Developer/LogVisualizer/Private/SVisualLoggerTimelineBar.cpp b/Engine/Source/Developer/LogVisualizer/Private/SVisualLoggerTimelineBar.cpp
--- a/Engine/Source/Developer/LogVisualizer/Private/SVisualLoggerTimelineBar.cpp
+++ b/Engine/Source/Developer/LogVisualizer/Private/SVisualLoggerTimelineBar.cpp
@@ -20,7 +20,7 @@ FReply SVisualLoggerTimelineBar::OnMouseButtonDown(const FGeometry& MyGeometry,
 	{
 		FName RowName = TimelineOwner.Pin()->GetName();
 		FVisualLoggerDBRow& DBRow = FVisualLoggerDatabase::Get().GetRowByName(RowName);
-		const float ScrubPosition = TimeSliderController->GetTimeSliderArgs().ScrubPosition.Get();
+		const double ScrubPosition = TimeSliderController->GetTimeSliderArgs().ScrubPosition.Get();
 		const int32 ClosestItem = DBRow.GetClosestItem(ScrubPosition);
 		const auto& Items = DBRow.GetItems();
 		if (Items.IsValidIndex(ClosestItem))
@@ -43,7 +43,7 @@ FReply SVisualLoggerTimelineBar::OnMouseButtonUp(const FGeometry& MyGeometry, co
 	{
 		FName RowName = TimelineOwner.Pin()->GetName();
 		FVisualLoggerDBRow& DBRow = FVisualLoggerDatabase::Get().GetRowByName(RowName);
-		const float ScrubPosition = TimeSliderController->GetTimeSliderArgs().ScrubPosition.Get();
+		const double ScrubPosition = TimeSliderController->GetTimeSliderArgs().ScrubPosition.Get();
 		const int32 ClosestItem = DBRow.GetClosestItem(ScrubPosition);
 		const auto& Items = DBRow.GetItems();
 		if (Items.IsValidIndex(ClosestItem))
@@ -59,7 +59,7 @@ FReply SVisualLoggerTimelineBar::OnMouseMove(const FGeometry& MyGeometry, const
 	FName RowName = TimelineOwner.Pin()->GetName();
 	FVisualLoggerDBRow& DBRow = FVisualLoggerDatabase::Get().GetRowByName(RowName);
 
-	const float ClosestMouseTime = TimeSliderController->GetTimeAtCursorPosition(MyGeometry, MouseEvent);
+	const double ClosestMouseTime = TimeSliderController->GetTimeAtCursorPosition(MyGeometry, MouseEvent);
 
 	const int32 NewItemIndex = DBRow.GetClosestItem(ClosestMouseTime);
 	const auto& Items = DBRow.GetItems();
@@ -125,7 +125,7 @@ void SVisualLoggerTimelineBar::Construct(const FArguments& InArgs, TSharedPtr<FV
 	TimeSliderController = InTimeSliderController;
 	TimelineOwner = InTimelineOwner;
 
-	TRange<float> LocalViewRange = TimeSliderController->GetTimeSliderArgs().ViewRange.Get();
+	TRange<double> LocalViewRange = TimeSliderController->GetTimeSliderArgs().ViewRange.Get();
 
 	MouseMoveClosestItemIndex = INDEX_NONE;
 }
@@ -143,13 +143,11 @@ int32 SVisualLoggerTimelineBar::OnPaint(const FPaintArgs& Args, const FGeometry&
 	FArrangedChildren ArrangedChildren(EVisibility::Visible);
 	ArrangeChildren(AllottedGeometry, ArrangedChildren);
 
-	TRange<float> LocalViewRange = TimeSliderController->GetTimeSliderArgs().ViewRange.Get();
-	float LocalScrubPosition = TimeSliderController->GetTimeSliderArgs().ScrubPosition.Get();
+	TRange<double> LocalViewRange = TimeSliderController->GetTimeSliderArgs().ViewRange.Get();
+	double LocalScrubPosition = TimeSliderController->GetTimeSliderArgs().ScrubPosition.Get();
 
-	float ViewRange = LocalViewRange.Size<float>();
-	float PixelsPerInput = ViewRange > 0 ? AllottedGeometry.GetLocalSize().X / ViewRange : 0;
-	float CurrentScrubLinePos = (LocalScrubPosition - LocalViewRange.GetLowerBoundValue()) * PixelsPerInput;
-	float BoxWidth = FMath::Max(0.08f * PixelsPerInput, 4.0f);
+	double ViewRange = LocalViewRange.Size<double>();
+	double PixelsPerInput = ViewRange > 0. ? AllottedGeometry.GetLocalSize().X / ViewRange : 0.;
 
 	// Draw a region around the entire section area
 	FSlateDrawElement::MakeBox(
@@ -170,8 +168,8 @@ int32 SVisualLoggerTimelineBar::OnPaint(const FPaintArgs& Args, const FGeometry&
 
 	const ESlateDrawEffect DrawEffects = ESlateDrawEffect::None;// bEnabled ? ESlateDrawEffect::None : ESlateDrawEffect::DisabledEffect;
 
-	TArray<float> ErrorTimes;
-	TArray<float> WarningTimes;
+	TArray<double> ErrorTimes;
+	TArray<double> WarningTimes;
 	int32 EntryIndex = 0;
 
 	FVisualLoggerDBRow& DBRow = FVisualLoggerDatabase::Get().GetRowByName(TimelineOwner.Pin()->GetName());
@@ -192,10 +190,9 @@ int32 SVisualLoggerTimelineBar::OnPaint(const FPaintArgs& Args, const FGeometry&
 		}
 
 		// find bar width, connect all contiguous bars to draw them as one geometry (rendering optimization)
-		const float StartPos = (Entry.TimeStamp - LocalViewRange.GetLowerBoundValue()) * PixelsPerInput - 2;
-		float EndPos = (Entry.TimeStamp - LocalViewRange.GetLowerBoundValue()) * PixelsPerInput + 2;
+		const double StartPos = (Entry.TimeStamp - LocalViewRange.GetLowerBoundValue()) * PixelsPerInput - 2;
+		double EndPos = (Entry.TimeStamp - LocalViewRange.GetLowerBoundValue()) * PixelsPerInput + 2;
 		int32 StartIndex = EntryIndex;
-		float LastEndX = MAX_FLT;
 		for (; StartIndex < Entries.Num(); ++StartIndex)
 		{
 			const FVisualLogEntry& CurrentEntry = Entries[StartIndex].Entry;
@@ -230,7 +227,7 @@ int32 SVisualLoggerTimelineBar::OnPaint(const FPaintArgs& Args, const FGeometry&
 				}
 			}
 
-			const float CurrentStartPos = (CurrentEntry.TimeStamp - LocalViewRange.GetLowerBoundValue()) * PixelsPerInput - 2;
+			const double CurrentStartPos = (CurrentEntry.TimeStamp - LocalViewRange.GetLowerBoundValue()) * PixelsPerInput - 2;
 			if (CurrentStartPos > EndPos)
 			{
 				break;
@@ -240,12 +237,12 @@ int32 SVisualLoggerTimelineBar::OnPaint(const FPaintArgs& Args, const FGeometry&
 
 		if (EndPos - StartPos > 0)
 		{
-			const float BarWidth = (EndPos - StartPos);
+			const double BarWidth = (EndPos - StartPos);
 			FSlateDrawElement::MakeBox(
 				OutDrawElements,
 				RetLayerId,
 				AllottedGeometry.ToPaintGeometry(
-				FVector2D(StartPos, 0.0f),
+				FVector2D(StartPos, 0.),
 				FVector2D(BarWidth, AllottedGeometry.Size.Y)),
 				FillImage,
 				DrawEffects,
@@ -258,14 +255,14 @@ int32 SVisualLoggerTimelineBar::OnPaint(const FPaintArgs& Args, const FGeometry&
 	if (WarningTimes.Num()) RetLayerId++;
 	for (auto CurrentTime : WarningTimes)
 	{
-		float LinePos = (CurrentTime - LocalViewRange.GetLowerBoundValue()) * PixelsPerInput;
+		double LinePos = (CurrentTime - LocalViewRange.GetLowerBoundValue()) * PixelsPerInput;
 
 		FSlateDrawElement::MakeBox(
 			OutDrawElements,
 			RetLayerId,
 			AllottedGeometry.ToPaintGeometry(
-			FVector2D(LinePos - 3, 0.0f),
-			FVector2D(6, AllottedGeometry.GetLocalSize().Y)),
+			FVector2D(LinePos - 3., 0.),
+			FVector2D(6., AllottedGeometry.GetLocalSize().Y)),
 			FillImage,
 			DrawEffects,
 			WarningTimeColor
@@ -275,13 +272,13 @@ int32 SVisualLoggerTimelineBar::OnPaint(const FPaintArgs& Args, const FGeometry&
 	if (ErrorTimes.Num()) RetLayerId++;
 	for (auto CurrentTime : ErrorTimes)
 	{
-		float LinePos = (CurrentTime - LocalViewRange.GetLowerBoundValue()) * PixelsPerInput;
+		double LinePos = (CurrentTime - LocalViewRange.GetLowerBoundValue()) * PixelsPerInput;
 
 		FSlateDrawElement::MakeBox(
 			OutDrawElements,
 			RetLayerId,
 			AllottedGeometry.ToPaintGeometry(
-			FVector2D(LinePos - 3, 0.0f),
+			FVector2D(LinePos - 3., 0.),
 			FVector2D(6, AllottedGeometry.GetLocalSize().Y)),
 			FillImage,
 			DrawEffects,
@@ -294,15 +291,15 @@ int32 SVisualLoggerTimelineBar::OnPaint(const FPaintArgs& Args, const FGeometry&
 	if (TimelineOwner.Pin()->IsSelected() && DBRow.GetCurrentItemIndex() != INDEX_NONE)
 	{
 		const auto &HighlightedItemEntry = DBRow.GetCurrentItem();
-		float CurrentTime = HighlightedItemEntry.Entry.TimeStamp;
-		float LinePos = (CurrentTime - LocalViewRange.GetLowerBoundValue()) * PixelsPerInput;
+		double CurrentTime = HighlightedItemEntry.Entry.TimeStamp;
+		double LinePos = (CurrentTime - LocalViewRange.GetLowerBoundValue()) * PixelsPerInput;
 
 		FSlateDrawElement::MakeBox(
 			OutDrawElements,
 			++RetLayerId,
 			AllottedGeometry.ToPaintGeometry(
-			FVector2D(LinePos - 2, 0.0f),
-			FVector2D(4, AllottedGeometry.GetLocalSize().Y)),
+			FVector2D(LinePos - 2., 0.),
+			FVector2D(4., AllottedGeometry.GetLocalSize().Y)),
 			SelectedFillImage,
 			ESlateDrawEffect::None,
 			SelectedBarColor
diff --git a/Engine/Source/Developer/LogVisualizer/Private/SVisualLoggerTimelinesContainer.cpp b/Engine/Source/Developer/LogVisualizer/Private/SVisualLoggerTimelinesContainer.cpp
--- a/Engine/Source/Developer/LogVisualizer/Private/SVisualLoggerTimelinesContainer.cpp
+++ b/Engine/Source/Developer/LogVisualizer/Private/SVisualLoggerTimelinesContainer.cpp
@@ -259,8 +259,8 @@ void SVisualLoggerTimelinesContainer::ResetData()
 	}
 	CachedSelectedTimelines.Reset();
 
-	CachedMinTime = FLT_MAX;
-	CachedMaxTime = 0;
+	CachedMinTime = MAX_dbl;
+	CachedMaxTime = 0.;
 	TimeSliderController->SetClampRange(0, 5);
 	TimeSliderController->SetTimeRange(0, 5);
 }
@@ -281,8 +281,8 @@ void SVisualLoggerTimelinesContainer::Construct(const FArguments& InArgs, TShare
 			]
 		];
 
-	CachedMinTime = FLT_MAX;
-	CachedMaxTime = 0;
+	CachedMinTime = MAX_dbl;
+	CachedMaxTime = 0.;
 
 	FVisualLoggerDatabase::Get().GetEvents().OnNewRow.AddRaw(this, &SVisualLoggerTimelinesContainer::OnNewRowHandler);
 	FVisualLoggerDatabase::Get().GetEvents().OnNewItem.AddRaw(this, &SVisualLoggerTimelinesContainer::OnNewItemHandler);
@@ -319,15 +319,12 @@ void SVisualLoggerTimelinesContainer::OnNewItemHandler(const FVisualLoggerDBRow&
 	CachedMinTime = CachedMinTime < Entry.Entry.TimeStamp ? CachedMinTime : Entry.Entry.TimeStamp;
 	CachedMaxTime = CachedMaxTime > Entry.Entry.TimeStamp ? CachedMaxTime : Entry.Entry.TimeStamp;
 
-	TRange<float> LocalViewRange = TimeSliderController->GetTimeSliderArgs().ViewRange.Get();
-	TRange<float> ClampRange = TimeSliderController->GetTimeSliderArgs().ClampRange.Get();
-	float ZoomLevel = LocalViewRange.Size<float>() / ClampRange.Size<float>();
+	TRange<double> LocalViewRange = TimeSliderController->GetTimeSliderArgs().ViewRange.Get();
+	TRange<double> ClampRange = TimeSliderController->GetTimeSliderArgs().ClampRange.Get();
+	double ZoomLevel = LocalViewRange.Size<double>() / ClampRange.Size<double>();
 
-	TimeSliderController->GetTimeSliderArgs().ClampRange = TRange<float>(CachedMinTime, CachedMaxTime + 0.1f);
-	//if ( FMath::Abs(ZoomLevel - 1) <= SMALL_NUMBER)
-	//{
-		TimeSliderController->GetTimeSliderArgs().ViewRange = TimeSliderController->GetTimeSliderArgs().ClampRange;
-	//}
+	TimeSliderController->GetTimeSliderArgs().ClampRange = TRange<double>(CachedMinTime, CachedMaxTime + 0.1);
+	TimeSliderController->GetTimeSliderArgs().ViewRange = TimeSliderController->GetTimeSliderArgs().ClampRange;
 }
 
 void SVisualLoggerTimelinesContainer::OnSearchChanged(const FText& Filter)
diff --git a/Engine/Source/Developer/LogVisualizer/Private/SVisualLoggerTimelinesContainer.h b/Engine/Source/Developer/LogVisualizer/Private/SVisualLoggerTimelinesContainer.h
--- a/Engine/Source/Developer/LogVisualizer/Private/SVisualLoggerTimelinesContainer.h
+++ b/Engine/Source/Developer/LogVisualizer/Private/SVisualLoggerTimelinesContainer.h
@@ -83,7 +83,8 @@ protected:
 
 	TArray<TSharedPtr<SLogVisualizerTimeline> > TimelineItems;
 	TArray<TSharedPtr<SLogVisualizerTimeline> > CachedSelectedTimelines;
-	float CachedMinTime, CachedMaxTime;
+	double CachedMinTime;
+	double CachedMaxTime;
 
 	FText CurrentSearchText;
 };
diff --git a/Engine/Source/Developer/LogVisualizer/Private/SVisualLoggerView.cpp b/Engine/Source/Developer/LogVisualizer/Private/SVisualLoggerView.cpp
--- a/Engine/Source/Developer/LogVisualizer/Private/SVisualLoggerView.cpp
+++ b/Engine/Source/Developer/LogVisualizer/Private/SVisualLoggerView.cpp
@@ -98,7 +98,7 @@ void SVisualLoggerView::Construct(const FArguments& InArgs, const TSharedRef<FUI
 	ChildSlot
 		[
 			SNew(SBorder)
-			.Padding(2)
+			.Padding(2.f)
 			.BorderImage(FLogVisualizerStyle::Get().GetBrush("ToolPanel.GroupBorder"))
 			[
 				SNew(SHorizontalBox)
@@ -136,7 +136,7 @@ void SVisualLoggerView::Construct(const FArguments& InArgs, const TSharedRef<FUI
 								.Visibility_Lambda([this]()->EVisibility{ return TimelinesContainer.IsValid() && (TimelinesContainer->GetAllNodes().Num() > 1 || FVisualLoggerFilters::Get().GetSelectedObjects().Num() > 0) ? EVisibility::Visible : EVisibility::Collapsed; })
 								.ComboButtonStyle(FLogVisualizerStyle::Get(), "Filters.Style")
 								.ForegroundColor(FLinearColor::White)
-								.ContentPadding(0)
+								.ContentPadding(0.f)
 								.OnGetMenuContent(this, &SVisualLoggerView::MakeClassesFilterMenu)
 								.ToolTipText(LOCTEXT("SetFilterByClasses", "Select classes to show"))
 								.HasDownArrow(true)
diff --git a/Engine/Source/Developer/LogVisualizer/Private/VisualLoggerCanvasRenderer.cpp b/Engine/Source/Developer/LogVisualizer/Private/VisualLoggerCanvasRenderer.cpp
--- a/Engine/Source/Developer/LogVisualizer/Private/VisualLoggerCanvasRenderer.cpp
+++ b/Engine/Source/Developer/LogVisualizer/Private/VisualLoggerCanvasRenderer.cpp
@@ -26,7 +26,7 @@ namespace LogVisualizer
 	{
 		if (PointInFrustum(Canvas, WorldLocation))
 		{
-			const FVector ScreenLocation = Canvas->Project(WorldLocation);
+			const FVector3f ScreenLocation = UE::LWC::NarrowWorldPositionChecked(Canvas->Project(WorldLocation));
 			Canvas->DrawText(Font, *TextToDraw, ScreenLocation.X, ScreenLocation.Y);
 		}
 	}
@@ -35,7 +35,7 @@ namespace LogVisualizer
 	{
 		if (PointInFrustum(Canvas, WorldLocation))
 		{
-			const FVector ScreenLocation = Canvas->Project(WorldLocation);
+			const FVector3f ScreenLocation = UE::LWC::NarrowWorldPositionChecked(Canvas->Project(WorldLocation));
 			float TextXL = 0.f;
 			float TextYL = 0.f;
 			Canvas->StrLen(Font, TextToDraw, TextXL, TextYL);
@@ -47,7 +47,7 @@ namespace LogVisualizer
 	{
 		if (PointInFrustum(Canvas, WorldLocation))
 		{
-			const FVector ScreenLocation = Canvas->Project(WorldLocation);
+			const FVector3f ScreenLocation = UE::LWC::NarrowWorldPositionChecked(Canvas->Project(WorldLocation));
 			float TextXL = 0.f;
 			float TextYL = 0.f;
 			Canvas->StrLen(Font, TextToDraw, TextXL, TextYL);
@@ -113,11 +113,11 @@ void FVisualLoggerCanvasRenderer::DrawOnCanvas(class UCanvas* Canvas, class APla
 	if (bDirtyData && FLogVisualizer::Get().GetTimeSliderController().IsValid())
 	{
 		const FVisualLoggerTimeSliderArgs&  TimeSliderArgs = FLogVisualizer::Get().GetTimeSliderController()->GetTimeSliderArgs();
-		TRange<float> LocalViewRange = TimeSliderArgs.ViewRange.Get();
-		const float LocalViewRangeMin = LocalViewRange.GetLowerBoundValue();
-		const float LocalViewRangeMax = LocalViewRange.GetUpperBoundValue();
-		const float LocalSequenceLength = LocalViewRangeMax - LocalViewRangeMin;
-		const float WindowHalfWidth = LocalSequenceLength * TimeSliderArgs.CursorSize.Get() * 0.5f;
+		TRange<double> LocalViewRange = TimeSliderArgs.ViewRange.Get();
+		const double LocalViewRangeMin = LocalViewRange.GetLowerBoundValue();
+		const double LocalViewRangeMax = LocalViewRange.GetUpperBoundValue();
+		const double LocalSequenceLength = LocalViewRangeMax - LocalViewRangeMin;
+		const double WindowHalfWidth = LocalSequenceLength * TimeSliderArgs.CursorSize.Get() * 0.5;
 		const FVector2D TimeStampWindow(SelectedEntry.TimeStamp - WindowHalfWidth, SelectedEntry.TimeStamp + WindowHalfWidth);
 
 		CollectedGraphs.Reset();
@@ -168,7 +168,7 @@ void FVisualLoggerCanvasRenderer::DrawOnCanvas(class UCanvas* Canvas, class APla
 						CollectedGraphData.Max.X = FMath::Max(CollectedGraphData.Max.X, SampleValue.X);
 						CollectedGraphData.Max.Y = FMath::Max(CollectedGraphData.Max.Y, SampleValue.Y);
 
-						const float CurrentTimeStamp = GraphData.TimeStamps[SampleIndex];
+						const double CurrentTimeStamp = GraphData.TimeStamps[SampleIndex];
 						if (CurrentTimeStamp < TimeStampWindow.X)
 						{
 							LineData.LeftExtreme = SampleValue;
@@ -208,11 +208,11 @@ void FVisualLoggerCanvasRenderer::DrawHistogramGraphs(class UCanvas* Canvas, cla
 	if (CollectedGraphs.Num() > 0)
 	{
 		const FVisualLoggerTimeSliderArgs&  TimeSliderArgs = FLogVisualizer::Get().GetTimeSliderController()->GetTimeSliderArgs();
-		TRange<float> LocalViewRange = TimeSliderArgs.ViewRange.Get();
-		const float LocalViewRangeMin = LocalViewRange.GetLowerBoundValue();
-		const float LocalViewRangeMax = LocalViewRange.GetUpperBoundValue();
-		const float LocalSequenceLength = LocalViewRangeMax - LocalViewRangeMin;
-		const float WindowHalfWidth = LocalSequenceLength * TimeSliderArgs.CursorSize.Get() * 0.5f;
+		TRange<double> LocalViewRange = TimeSliderArgs.ViewRange.Get();
+		const double LocalViewRangeMin = LocalViewRange.GetLowerBoundValue();
+		const double LocalViewRangeMax = LocalViewRange.GetUpperBoundValue();
+		const double LocalSequenceLength = LocalViewRangeMax - LocalViewRangeMin;
+		const double WindowHalfWidth = LocalSequenceLength * TimeSliderArgs.CursorSize.Get() * 0.5f;
 		const FVector2D TimeStampWindow(SelectedEntry.TimeStamp - WindowHalfWidth, SelectedEntry.TimeStamp + WindowHalfWidth);
 
 		const FColor GraphsBackgroundColor = ULogVisualizerSettings::StaticClass()->GetDefaultObject<ULogVisualizerSettings>()->GraphsBackgroundColor;
@@ -232,7 +232,7 @@ void FVisualLoggerCanvasRenderer::DrawHistogramGraphs(class UCanvas* Canvas, cla
 		const float YGraphSpacing = 0.2f / (MaxNumberOfGraphs + 1);
 
 		const float StartX = XGraphSpacing;
-		float StartY = 0.5 + (0.5 * NumberOfRows - 1) * (GraphHeight + YGraphSpacing);
+		float StartY = 0.5f + (0.5f * NumberOfRows - 1) * (GraphHeight + YGraphSpacing);
 
 		float CurrentX = StartX;
 		float CurrentY = StartY;
@@ -278,7 +278,7 @@ void FVisualLoggerCanvasRenderer::DrawHistogramGraphs(class UCanvas* Canvas, cla
 					Hue -= FMath::FloorToFloat(Hue);
 				}
 
-				HistogramGraph->GetGraphLine(LineIndex)->Color = FLinearColor::MakeFromHSV8(Hue * 255, 200, 244);
+				HistogramGraph->GetGraphLine(LineIndex)->Color = FLinearColor::MakeFromHSV8(static_cast<uint8>(Hue * 255), 200, 244);
 				HistogramGraph->GetGraphLine(LineIndex)->LineName = DataName;
 				HistogramGraph->GetGraphLine(LineIndex)->Data.Append(LinesIt->Value.Samples);
 				HistogramGraph->GetGraphLine(LineIndex)->LeftExtreme = LinesIt->Value.LeftExtreme;
@@ -296,15 +296,15 @@ void FVisualLoggerCanvasRenderer::DrawHistogramGraphs(class UCanvas* Canvas, cla
 					{
 						if (SampleData.X >= TimeStampWindow.X && SampleData.X <= TimeStampWindow.Y)
 						{
-							It->Value.Min.Y = FMath::Min<float>(It->Value.Min.Y, SampleData.Y);
-							It->Value.Max.Y = FMath::Max<float>(It->Value.Max.Y, SampleData.Y);
+							It->Value.Min.Y = FMath::Min(It->Value.Min.Y, SampleData.Y);
+							It->Value.Max.Y = FMath::Max(It->Value.Max.Y, SampleData.Y);
 						}
 					}
 				}
 			}
 
 			FVector2D GraphSpaceSize;
-			GraphSpaceSize.Y = GraphSpaceSize.X = 0.8f / CollectedGraphs.Num();
+			GraphSpaceSize.Y = GraphSpaceSize.X = 0.8 / CollectedGraphs.Num();
 
 			HistogramGraph->SetGraphScreenSize(CurrentX, CurrentX + GraphWidth, CurrentY, CurrentY + GraphHeight);
 			CurrentX += GraphWidth + XGraphSpacing;
@@ -312,7 +312,7 @@ void FVisualLoggerCanvasRenderer::DrawHistogramGraphs(class UCanvas* Canvas, cla
 
 			HistogramGraph->DrawCursorOnGraph(true);
 			HistogramGraph->UseTinyFont(CollectedGraphs.Num() >= 5);
-			HistogramGraph->SetCursorLocation(SelectedEntry.TimeStamp);
+			HistogramGraph->SetCursorLocation(FloatCastChecked<float>(SelectedEntry.TimeStamp, /* Precision */ 1./16.));
 			HistogramGraph->SetNumThresholds(0);
 			HistogramGraph->SetStyles(EGraphAxisStyle::Grid, EGraphDataStyle::Lines);
 			HistogramGraph->SetBackgroundColor(GraphsBackgroundColor);
diff --git a/Engine/Source/Developer/LogVisualizer/Private/VisualLoggerDatabase.cpp b/Engine/Source/Developer/LogVisualizer/Private/VisualLoggerDatabase.cpp
--- a/Engine/Source/Developer/LogVisualizer/Private/VisualLoggerDatabase.cpp
+++ b/Engine/Source/Developer/LogVisualizer/Private/VisualLoggerDatabase.cpp
@@ -47,10 +47,10 @@ void FVisualLoggerDBRow::SetItemVisibility(int32 ItemIndex, bool IsVisible)
 	}
 }
 
-int32 FVisualLoggerDBRow::GetClosestItem(float Time) const
+int32 FVisualLoggerDBRow::GetClosestItem(double Time) const
 {
 	int32 BestItemIndex = INDEX_NONE;
-	float BestDistance = MAX_FLT;
+	double BestDistance = MAX_dbl;
 	for (int32 Index = 0; Index < Items.Num(); Index++)
 	{
 		auto& CurrentEntryItem = Items[Index];
@@ -60,7 +60,7 @@ int32 FVisualLoggerDBRow::GetClosestItem(float Time) const
 			continue;
 		}
 		TArray<FVisualLoggerCategoryVerbosityPair> OutCategories;
-		const float CurrentDist = FMath::Abs(Time - CurrentEntryItem.Entry.TimeStamp);
+		const double CurrentDist = FMath::Abs(Time - CurrentEntryItem.Entry.TimeStamp);
 		if (CurrentDist < BestDistance)
 		{
 			BestDistance = CurrentDist;
@@ -68,7 +68,7 @@ int32 FVisualLoggerDBRow::GetClosestItem(float Time) const
 		}
 	}
 
-	const float CurrentDist = Items.IsValidIndex(CurrentItemIndex) && IsItemVisible(CurrentItemIndex) ? FMath::Abs(Time - Items[CurrentItemIndex].Entry.TimeStamp) : MAX_FLT;
+	const double CurrentDist = Items.IsValidIndex(CurrentItemIndex) && IsItemVisible(CurrentItemIndex) ? FMath::Abs(Time - Items[CurrentItemIndex].Entry.TimeStamp) : MAX_dbl;
 	if (BestItemIndex != INDEX_NONE && CurrentDist > BestDistance)
 	{
 		return BestItemIndex;
@@ -77,10 +77,10 @@ int32 FVisualLoggerDBRow::GetClosestItem(float Time) const
 	return CurrentItemIndex;
 }
 
-int32 FVisualLoggerDBRow::GetClosestItem(float Time, float ScrubTime) const
+int32 FVisualLoggerDBRow::GetClosestItem(double Time, double ScrubTime) const
 {
 	int32 BestItemIndex = INDEX_NONE;
-	float BestDistance = MAX_FLT;
+	double BestDistance = MAX_dbl;
 	for (int32 Index = 0; Index < Items.Num(); Index++)
 	{
 		auto& CurrentEntryItem = Items[Index];
@@ -95,7 +95,7 @@ int32 FVisualLoggerDBRow::GetClosestItem(float Time, float ScrubTime) const
 			continue;
 		}
 		TArray<FVisualLoggerCategoryVerbosityPair> OutCategories;
-		const float CurrentDist = FMath::Abs(CurrentEntryItem.Entry.TimeStamp - Time);
+		const double CurrentDist = FMath::Abs(CurrentEntryItem.Entry.TimeStamp - Time);
 		if (CurrentDist < BestDistance && CurrentEntryItem.Entry.TimeStamp <= Time)
 		{
 			BestDistance = CurrentDist;
@@ -103,9 +103,6 @@ int32 FVisualLoggerDBRow::GetClosestItem(float Time, float ScrubTime) const
 		}
 	}
 
-
-	const float CurrentDist = Items.IsValidIndex(CurrentItemIndex) ? FMath::Abs(Items[CurrentItemIndex].Entry.TimeStamp - Time) : MAX_FLT;
-
 	if (BestItemIndex != INDEX_NONE)
 	{
 		return BestItemIndex;
@@ -176,10 +173,10 @@ bool FVisualLoggerDatabase::ContainsRowByName(FName InName)
 
 FVisualLoggerDBRow& FVisualLoggerDatabase::GetRowByName(FName InName)
 {
-	const bool bContainsRow = RowNameToIndex.Contains(InName);
-	check(bContainsRow);
+	const int32* Idx = RowNameToIndex.Find(InName);
+	check(Idx);
 
-	return Rows[RowNameToIndex[InName]];
+	return Rows[*Idx];
 }
 
 void FVisualLoggerDatabase::SelectRow(FName InName, bool bDeselectOtherNodes)
diff --git a/Engine/Source/Developer/LogVisualizer/Private/VisualLoggerDatabase.h b/Engine/Source/Developer/LogVisualizer/Private/VisualLoggerDatabase.h
--- a/Engine/Source/Developer/LogVisualizer/Private/VisualLoggerDatabase.h
+++ b/Engine/Source/Developer/LogVisualizer/Private/VisualLoggerDatabase.h
@@ -51,8 +51,8 @@ public:
 	bool IsItemVisible(int32 ItemIndex) const { return HiddenItems[ItemIndex] == false; }
 	int32 GetNumberOfHiddenItems() const { return NumHiddenItems; }
 
-	int32 GetClosestItem(float Time) const;
-	int32 GetClosestItem(float Time, float ScrubTime) const;
+	int32 GetClosestItem(double Time) const;
+	int32 GetClosestItem(double Time, double ScrubTime) const;
 
 protected:
 	FVisualLoggerDBEvents& DBEvents;
@@ -112,7 +112,7 @@ struct FVisualLoggerGraphData
 {
 	FName DataName;
 	TArray<FVector2D> Samples;
-	TArray<float> TimeStamps;
+	TArray<double> TimeStamps;
 
 	FVisualLoggerGraphData(FName InDataName) : DataName(InDataName) {}
 };
diff --git a/Engine/Source/Developer/LogVisualizer/Private/VisualLoggerHUD.cpp b/Engine/Source/Developer/LogVisualizer/Private/VisualLoggerHUD.cpp
--- a/Engine/Source/Developer/LogVisualizer/Private/VisualLoggerHUD.cpp
+++ b/Engine/Source/Developer/LogVisualizer/Private/VisualLoggerHUD.cpp
@@ -63,7 +63,7 @@ void AVisualLoggerHUD::PostRender()
 				TextItem.Text = LOCTEXT("NotActorUnderCursor", "Not actor under cursor" );
 			}
 			Canvas->DrawItem( TextItem, X, Y );
-			Y += TextItem.DrawnSize.Y;
+			Y += static_cast<float>(TextItem.DrawnSize.Y);
 
 			if (DebugCamController->PickedActor.IsValid())
 			{
diff --git a/Engine/Source/Developer/LogVisualizer/Private/VisualLoggerRenderingActorBase.cpp b/Engine/Source/Developer/LogVisualizer/Private/VisualLoggerRenderingActorBase.cpp
--- a/Engine/Source/Developer/LogVisualizer/Private/VisualLoggerRenderingActorBase.cpp
+++ b/Engine/Source/Developer/LogVisualizer/Private/VisualLoggerRenderingActorBase.cpp
@@ -12,7 +12,7 @@
 
 namespace FDebugDrawing
 {
-	const FVector NavOffset(0, 0, 15);
+	const FVector NavOffset(0., 0., 15.);
 }
 
 class UVisualLoggerRenderingComponent;
@@ -110,8 +110,8 @@ namespace
 		if (Verts.Num() >= 3)
 		{
 			const FVector SurfaceNormal = FVector::CrossProduct(Verts[1] - Verts[0], Verts[2] - Verts[0]);
-			const float TestDot = FVector::DotProduct(SurfaceNormal, FVector(0, 0, 1));
-			return TestDot > 0;
+			const double TestDot = FVector::DotProduct(SurfaceNormal, FVector::UpVector);
+			return TestDot > 0.;
 		}
 
 		return false;
@@ -122,7 +122,7 @@ namespace
 		TestMesh.Color = ElementToDraw->GetFColor();
 
 		FClipSMPolygon InPoly(ElementToDraw->Points.Num());
-		InPoly.FaceNormal = FVector(0, 0, 1);
+		InPoly.FaceNormal = FVector::UpVector;
 
 		const bool bHasCorrectWinding = IsPolygonWindingCorrect(ElementToDraw->Points);
 		if (bHasCorrectWinding)
@@ -174,7 +174,7 @@ void AVisualLoggerRenderingActorBase::GetDebugShapes(const FVisualLogEntry& InEn
 	if (bAddEntryLocationPointer)
 	{
 		constexpr float Length = 100;
-		const FVector DirectionNorm = FVector(0, 0, 1).GetSafeNormal();
+		const FVector DirectionNorm = FVector::UpVector;
 		FVector YAxis, ZAxis;
 		DirectionNorm.FindBestAxisVectors(YAxis, ZAxis);
 		DebugShapes.Cones.Add(FDebugRenderSceneProxy::FCone(FScaleMatrix(FVector(Length)) * FMatrix(DirectionNorm, YAxis, ZAxis, Entry->Location), 5, 5, FColor::Red));
@@ -243,8 +243,8 @@ void AVisualLoggerRenderingActorBase::GetDebugShapes(const FVisualLogEntry& InEn
 		{
 			struct FHeaderData
 			{
-				float VerticesNum, FacesNum;
-				FHeaderData(const FVector& InVector) : VerticesNum(InVector.X), FacesNum(InVector.Y) {}
+				int32 VerticesNum, FacesNum;
+				FHeaderData(const FVector& InVector) : VerticesNum(static_cast<int32>(InVector.X)), FacesNum(static_cast<int32>(InVector.Y)) {}
 			};
 			const FHeaderData HeaderData(ElementToDraw->Points[0]);
 
@@ -263,9 +263,9 @@ void AVisualLoggerRenderingActorBase::GetDebugShapes(const FVisualLogEntry& InEn
 			for (int32 VIdx = StartIndex; VIdx < EndIndex; VIdx++)
 			{
 				const FVector &CurrentFace = ElementToDraw->Points[VIdx];
-				TestMesh.Indices.Add(CurrentFace.X);
-				TestMesh.Indices.Add(CurrentFace.Y);
-				TestMesh.Indices.Add(CurrentFace.Z);
+				TestMesh.Indices.Add(static_cast<uint32>(CurrentFace.X));
+				TestMesh.Indices.Add(static_cast<uint32>(CurrentFace.Y));
+				TestMesh.Indices.Add(static_cast<uint32>(CurrentFace.Z));
 			}
 			DebugShapes.Meshes.Add(TestMesh);
 		}
@@ -330,11 +330,11 @@ void AVisualLoggerRenderingActorBase::GetDebugShapes(const FVisualLogEntry& InEn
 				const FVector Origin = ElementToDraw->Points[Index];
 				const FVector Direction = ElementToDraw->Points[Index + 1].GetSafeNormal();
 				const FVector Angles = ElementToDraw->Points[Index + 2];
-				const float Length = Angles.X;
+				const double Length = Angles.X;
 
 				FVector YAxis, ZAxis;
 				Direction.FindBestAxisVectors(YAxis, ZAxis);
-				DebugShapes.Cones.Add(FDebugRenderSceneProxy::FCone(FScaleMatrix(FVector(Length)) * FMatrix(Direction, YAxis, ZAxis, Origin), Angles.Y, Angles.Z, Color));
+				DebugShapes.Cones.Add(FDebugRenderSceneProxy::FCone(FScaleMatrix(FVector(Length)) * FMatrix(Direction, YAxis, ZAxis, Origin), static_cast<float>(Angles.Y), static_cast<float>(Angles.Z), Color));
 
 				if (bDrawLabel)
 				{
@@ -352,9 +352,11 @@ void AVisualLoggerRenderingActorBase::GetDebugShapes(const FVisualLogEntry& InEn
 				const FVector Start = ElementToDraw->Points[Index];
 				const FVector End = ElementToDraw->Points[Index + 1];
 				const FVector OtherData = ElementToDraw->Points[Index + 2];
-				const float HalfHeight = 0.5f * (End - Start).Size();
-				const FVector Center = 0.5f * (Start + End);
-				DebugShapes.Cylinders.Add(FDebugRenderSceneProxy::FWireCylinder(Center, OtherData.X, HalfHeight, Color)); // Base parameter is the center of the cylinder
+				const float HalfHeight = FloatCastChecked<float>(0.5 * (End - Start).Size(), UE::LWC::DefaultFloatPrecision);
+				const FVector Center = 0.5 * (Start + End);
+				DebugShapes.Cylinders.Add(FDebugRenderSceneProxy::FWireCylinder(Center
+					, static_cast<float>(OtherData.X)
+					, HalfHeight, Color)); // Base parameter is the center of the cylinder
 				if (bDrawLabel)
 				{
 					DebugShapes.Texts.Add(FDebugRenderSceneProxy::FText3d(ElementToDraw->Description, Center, Color));
@@ -371,8 +373,8 @@ void AVisualLoggerRenderingActorBase::GetDebugShapes(const FVisualLogEntry& InEn
 				const FVector Base = ElementToDraw->Points[Index + 0];
 				const FVector FirstData = ElementToDraw->Points[Index + 1];
 				const FVector SecondData = ElementToDraw->Points[Index + 2];
-				const float HalfHeight = FirstData.X;
-				const float Radius = FirstData.Y;
+				const float HalfHeight = static_cast<float>(FirstData.X);
+				const float Radius = static_cast<float>(FirstData.Y);
 				const FQuat Rotation = FQuat(FirstData.Z, SecondData.X, SecondData.Y, SecondData.Z);
 
 				const FMatrix Axes = FQuatRotationTranslationMatrix(Rotation, FVector::ZeroVector);
@@ -398,7 +400,7 @@ void AVisualLoggerRenderingActorBase::GetDebugShapes(const FVisualLogEntry& InEn
 			struct FHeaderData
 			{
 				float MinZ, MaxZ;
-				FHeaderData(const FVector& InVector) : MinZ(InVector.X), MaxZ(InVector.Y) {}
+				FHeaderData(const FVector& InVector) : MinZ(FloatCastChecked<float>(InVector.X, UE::LWC::DefaultFloatPrecision)), MaxZ(FloatCastChecked<float>(InVector.Y, UE::LWC::DefaultFloatPrecision)) {}
 			};
 			const FHeaderData HeaderData(ElementToDraw->Points[0]);
 
@@ -410,27 +412,36 @@ void AVisualLoggerRenderingActorBase::GetDebugShapes(const FVisualLogEntry& InEn
 			int32 CurrentIndex = 0;
 			FDebugRenderSceneProxy::FMesh TestMesh;
 			TestMesh.Color = ElementToDraw->GetFColor();
+			TestMesh.Vertices.Reserve((AreaMeshPoints.Num() - 1) * 6);
+			TestMesh.Indices.Reserve((AreaMeshPoints.Num() - 1) * 6);
 
 			for (int32 PointIndex = 0; PointIndex < AreaMeshPoints.Num() - 1; PointIndex++)
 			{
 				FVector Point = AreaMeshPoints[PointIndex];
 				FVector NextPoint = AreaMeshPoints[PointIndex + 1];
 
-				FVector3f P1(Point.X, Point.Y, HeaderData.MinZ);
-				FVector3f P2(Point.X, Point.Y, HeaderData.MaxZ);
-				FVector3f P3(NextPoint.X, NextPoint.Y, HeaderData.MinZ);
-				FVector3f P4(NextPoint.X, NextPoint.Y, HeaderData.MaxZ);
-
-				TestMesh.Vertices.Add(P1); TestMesh.Vertices.Add(P2); TestMesh.Vertices.Add(P3);
-				TestMesh.Indices.Add(CurrentIndex + 0);
-				TestMesh.Indices.Add(CurrentIndex + 1);
-				TestMesh.Indices.Add(CurrentIndex + 2);
-				CurrentIndex += 3;
-				TestMesh.Vertices.Add(P3); TestMesh.Vertices.Add(P2); TestMesh.Vertices.Add(P4);
-				TestMesh.Indices.Add(CurrentIndex + 0);
-				TestMesh.Indices.Add(CurrentIndex + 1);
-				TestMesh.Indices.Add(CurrentIndex + 2);
-				CurrentIndex += 3;
+				
+				// LWC_TODO These won't work well with very large coordinates!
+				FVector3f P1(UE::LWC::NarrowWorldPositionChecked(Point.X, Point.Y, HeaderData.MinZ));
+				FVector3f P2(UE::LWC::NarrowWorldPositionChecked(Point.X, Point.Y, HeaderData.MaxZ));
+				FVector3f P3(UE::LWC::NarrowWorldPositionChecked(NextPoint.X, NextPoint.Y, HeaderData.MinZ));
+				FVector3f P4(UE::LWC::NarrowWorldPositionChecked(NextPoint.X, NextPoint.Y, HeaderData.MaxZ));
+
+				TestMesh.Vertices.Add(P1); 
+				TestMesh.Vertices.Add(P2); 
+				TestMesh.Vertices.Add(P3);
+
+				TestMesh.Indices.Add(CurrentIndex++);
+				TestMesh.Indices.Add(CurrentIndex++);
+				TestMesh.Indices.Add(CurrentIndex++);
+
+				TestMesh.Vertices.Add(P3); 
+				TestMesh.Vertices.Add(P2); 
+				TestMesh.Vertices.Add(P4);
+
+				TestMesh.Indices.Add(CurrentIndex++);
+				TestMesh.Indices.Add(CurrentIndex++);
+				TestMesh.Indices.Add(CurrentIndex++);
 			}
 			DebugShapes.Meshes.Add(TestMesh);
 
@@ -440,15 +451,15 @@ void AVisualLoggerRenderingActorBase::GetDebugShapes(const FVisualLogEntry& InEn
 				PolygonToDraw.SetColor(ElementToDraw->GetFColor());
 				PolygonToDraw.Points.Reserve(AreaMeshPoints.Num());
 				PolygonToDraw.Points = AreaMeshPoints;
-				GetPolygonMesh(&PolygonToDraw, PolygonMesh, FVector3f(0, 0, HeaderData.MaxZ));
+				GetPolygonMesh(&PolygonToDraw, PolygonMesh, FVector3f(0., 0., HeaderData.MaxZ));
 				DebugShapes.Meshes.Add(PolygonMesh);
 			}
 
 			for (int32 VIdx = 0; VIdx < AreaMeshPoints.Num(); VIdx++)
 			{
 				DebugShapes.Lines.Add(FDebugRenderSceneProxy::FDebugLine(
-					AreaMeshPoints[VIdx] + FVector(0, 0, HeaderData.MaxZ),
-					AreaMeshPoints[(VIdx + 1) % AreaMeshPoints.Num()] + FVector(0, 0, HeaderData.MaxZ),
+					AreaMeshPoints[VIdx] + FVector(0., 0., HeaderData.MaxZ),
+					AreaMeshPoints[(VIdx + 1) % AreaMeshPoints.Num()] + FVector(0., 0., HeaderData.MaxZ),
 					ElementToDraw->GetFColor(),
 					2)
 					);
@@ -485,7 +496,7 @@ void AVisualLoggerRenderingActorBase::GetDebugShapes(const FVisualLogEntry& InEn
 			{
 				const FVector Center = ElementToDraw->Points[Index + 0];
 				const FVector UpAxis = ElementToDraw->Points[Index + 1];
-				const float Radius = ElementToDraw->Points[Index + 2].X;
+				const double Radius = ElementToDraw->Points[Index + 2].X;
 				const float Thickness = float(ElementToDraw->Thicknes);
 
 				const FQuat Rotation = FQuat::FindBetweenNormals(FVector::UpVector, UpAxis);
@@ -496,7 +507,7 @@ void AVisualLoggerRenderingActorBase::GetDebugShapes(const FVisualLogEntry& InEn
 				FVector PrevPosition = FVector::ZeroVector;
 				for (int32 Div = 0; Div <= CircleDivs; Div++)
 				{
-					const float Angle = (float)Div / (float)CircleDivs * PI * 2.0f;
+					const float Angle = (float)Div / (float)CircleDivs * UE_PI * 2.0f;
 					const FVector Position = Center + (FMath::Cos(Angle) * XAxis + FMath::Sin(Angle) * YAxis) * Radius;
 					if (Div > 0)
 					{
diff --git a/Engine/Source/Developer/LogVisualizer/Private/VisualLoggerTimeSliderController.cpp b/Engine/Source/Developer/LogVisualizer/Private/VisualLoggerTimeSliderController.cpp
--- a/Engine/Source/Developer/LogVisualizer/Private/VisualLoggerTimeSliderController.cpp
+++ b/Engine/Source/Developer/LogVisualizer/Private/VisualLoggerTimeSliderController.cpp
@@ -26,28 +26,28 @@ struct FVisualLoggerTimeSliderController::FScrubRangeToScreen
 {
 	FVector2D WidgetSize;
 
-	TRange<float> ViewInput;
-	float ViewInputRange;
-	float PixelsPerInput;
+	TRange<double> ViewInput;
+	double ViewInputRange;
+	double PixelsPerInput;
 
-	FScrubRangeToScreen(TRange<float> InViewInput, const FVector2D& InWidgetSize )
+	FScrubRangeToScreen(TRange<double> InViewInput, const FVector2D& InWidgetSize )
 	{
 		WidgetSize = InWidgetSize;
 
 		ViewInput = InViewInput;
-		ViewInputRange = ViewInput.Size<float>();
+		ViewInputRange = ViewInput.Size<double>();
 		PixelsPerInput = ViewInputRange > 0 ? ( WidgetSize.X / ViewInputRange ) : 0;
 	}
 
 	/** Local Widget Space -> Curve Input domain. */
-	float LocalXToInput(float ScreenX) const
+	double LocalXToInput(double ScreenX) const
 	{
-		float LocalX = ScreenX;
+		double LocalX = ScreenX;
 		return (LocalX/PixelsPerInput) + ViewInput.GetLowerBoundValue();
 	}
 
 	/** Curve Input domain -> local Widget Space */
-	float InputToLocalX(float Input) const
+	double InputToLocalX(double Input) const
 	{
 		return (Input - ViewInput.GetLowerBoundValue()) * PixelsPerInput;
 	}
@@ -75,7 +75,7 @@ static float GetNextSpacing( uint32 CurrentStep )
 
 FVisualLoggerTimeSliderController::FVisualLoggerTimeSliderController(const FVisualLoggerTimeSliderArgs& InArgs)
 	: TimeSliderArgs( InArgs )
-	, DistanceDragged( 0.0f )
+	, DistanceDragged( 0. )
 	, bDraggingScrubber( false )
 	, bPanning( false )
 {
@@ -84,15 +84,15 @@ FVisualLoggerTimeSliderController::FVisualLoggerTimeSliderController(const FVisu
 	CursorBackground = FAppStyle::GetBrush("Sequencer.SectionArea.Background");
 }
 
-float FVisualLoggerTimeSliderController::DetermineOptimalSpacing(float InPixelsPerInput, uint32 MinTick, float MinTickSpacing) const
+double FVisualLoggerTimeSliderController::DetermineOptimalSpacing(double InPixelsPerInput, uint32 MinTick, float MinTickSpacing) const
 {
-	if (InPixelsPerInput == 0.0f)
+	if (InPixelsPerInput == 0.)
 		return MinTickSpacing;
 
 	uint32 CurStep = 0;
 
 	// Start with the smallest spacing
-	float Spacing = MinTickSpacing;
+	double Spacing = MinTickSpacing;
 
 	while( Spacing * InPixelsPerInput < MinTick )
 	{
@@ -117,9 +117,9 @@ struct FVisualLoggerTimeSliderController::FDrawTickArgs
 	/** Color of each tick */
 	FLinearColor TickColor;
 	/** Offset in Y where to start the tick */
-	float TickOffset;
+	double TickOffset;
 	/** Height in of major ticks */
-	float MajorTickHeight;
+	double MajorTickHeight;
 	/** Start layer for elements */
 	int32 StartLayer;
 	/** Draw effects to apply */
@@ -133,7 +133,7 @@ struct FVisualLoggerTimeSliderController::FDrawTickArgs
 
 void FVisualLoggerTimeSliderController::DrawTicks( FSlateWindowElementList& OutDrawElements, const struct FScrubRangeToScreen& RangeToScreen, FDrawTickArgs& InArgs ) const
 {
-	const float Spacing = DetermineOptimalSpacing( RangeToScreen.PixelsPerInput, ScrubConstants::MinPixelsPerDisplayTick, ScrubConstants::MinDisplayTickSpacing );
+	const double Spacing = DetermineOptimalSpacing( RangeToScreen.PixelsPerInput, ScrubConstants::MinPixelsPerDisplayTick, ScrubConstants::MinDisplayTickSpacing );
 
 	// Sub divisions
 	// @todo Sequencer may need more robust calculation
@@ -141,18 +141,18 @@ void FVisualLoggerTimeSliderController::DrawTicks( FSlateWindowElementList& OutD
 	// For slightly larger halfway tick mark
 	const int32 HalfDivider = Divider / 2;
 	// Find out where to start from
-	int32 OffsetNum = FMath::FloorToInt(RangeToScreen.ViewInput.GetLowerBoundValue() / Spacing);
+	int32 OffsetNum = IntCastChecked<int32>(FMath::FloorToInt(RangeToScreen.ViewInput.GetLowerBoundValue() / Spacing));
 	
 	FSlateFontInfo SmallLayoutFont = FCoreStyle::GetDefaultFontStyle("Regular", 8);
 
 	TArray<FVector2D> LinePoints;
 	LinePoints.AddUninitialized(2);
 
-	float Seconds = 0;
+	double Seconds = 0;
 	while( (Seconds = OffsetNum*Spacing) < RangeToScreen.ViewInput.GetUpperBoundValue() )
 	{
 		// X position local to start of the widget area
-		float XPos = RangeToScreen.InputToLocalX( Seconds );
+		double XPos = RangeToScreen.InputToLocalX( Seconds );
 		uint32 AbsOffsetNum = FMath::Abs(OffsetNum);
 
 		if ( AbsOffsetNum % Divider == 0 )
@@ -184,7 +184,7 @@ void FVisualLoggerTimeSliderController::DrawTicks( FSlateWindowElementList& OutD
 				// Space the text between the tick mark but slightly above
 				const TSharedRef< FSlateFontMeasure > FontMeasureService = FSlateApplication::Get().GetRenderer()->GetFontMeasureService();
 				FVector2D TextSize = FontMeasureService->Measure(FrameString, SmallLayoutFont);
-				FVector2D TextOffset( XPos-(TextSize.X*0.5f), InArgs.bMirrorLabels ? TextSize.Y :  FMath::Abs( InArgs.AllottedGeometry.GetLocalSize().Y - (InArgs.MajorTickHeight+TextSize.Y) ) );
+				FVector2D TextOffset( XPos-(TextSize.X*0.5), InArgs.bMirrorLabels ? TextSize.Y :  FMath::Abs( InArgs.AllottedGeometry.GetLocalSize().Y - (InArgs.MajorTickHeight+TextSize.Y) ) );
 
 				FSlateDrawElement::MakeText(
 					OutDrawElements,
@@ -231,17 +231,17 @@ int32 FVisualLoggerTimeSliderController::OnPaintTimeSlider( bool bMirrorLabels,
 	const bool bEnabled = bParentEnabled;
 	const ESlateDrawEffect DrawEffects = bEnabled ? ESlateDrawEffect::None : ESlateDrawEffect::DisabledEffect;
 
-	TRange<float> LocalViewRange = TimeSliderArgs.ViewRange.Get();
-	const float LocalViewRangeMin = LocalViewRange.GetLowerBoundValue();
-	const float LocalViewRangeMax = LocalViewRange.GetUpperBoundValue();
-	const float LocalSequenceLength = LocalViewRangeMax-LocalViewRangeMin;
+	TRange<double> LocalViewRange = TimeSliderArgs.ViewRange.Get();
+	const double LocalViewRangeMin = LocalViewRange.GetLowerBoundValue();
+	const double LocalViewRangeMax = LocalViewRange.GetUpperBoundValue();
+	const double LocalSequenceLength = LocalViewRangeMax-LocalViewRangeMin;
 	
 	FVector2D Scale = FVector2D(1.0f,1.0f);
 	if ( LocalSequenceLength > 0)
 	{
 		FScrubRangeToScreen RangeToScreen( LocalViewRange, AllottedGeometry.GetLocalSize() );
 	
-		const float MajorTickHeight = 9.0f;
+		const float MajorTickHeight = 9.;
 	
 		FDrawTickArgs Args;
 		Args.AllottedGeometry = AllottedGeometry;
@@ -251,7 +251,7 @@ int32 FVisualLoggerTimeSliderController::OnPaintTimeSlider( bool bMirrorLabels,
 		Args.ClippingRect = MyCullingRect;
 		Args.DrawEffects = DrawEffects;
 		Args.StartLayer = LayerId;
-		Args.TickOffset = bMirrorLabels ? 0.0f : FMath::Abs( AllottedGeometry.GetLocalSize().Y - MajorTickHeight );
+		Args.TickOffset = bMirrorLabels ? 0. : FMath::Abs( AllottedGeometry.GetLocalSize().Y - MajorTickHeight );
 		Args.MajorTickHeight = MajorTickHeight;
 
 		DrawTicks( OutDrawElements, RangeToScreen, Args );
@@ -260,12 +260,12 @@ int32 FVisualLoggerTimeSliderController::OnPaintTimeSlider( bool bMirrorLabels,
 		float HalfSize = FMath::TruncToFloat(HandleSize/2.0f);
 
 		// Draw the scrub handle
-		const float XPos = RangeToScreen.InputToLocalX( TimeSliderArgs.ScrubPosition.Get() );
+		const double XPos = RangeToScreen.InputToLocalX( TimeSliderArgs.ScrubPosition.Get() );
 
 		// Draw cursor size
 		const float CursorHalfSize = TimeSliderArgs.CursorSize.Get() * 0.5f;
 		const int32 CursorLayer = LayerId + 2;
-		const float CursorHalfLength = AllottedGeometry.GetLocalSize().X * CursorHalfSize;
+		const double CursorHalfLength = AllottedGeometry.GetLocalSize().X * CursorHalfSize;
 		FPaintGeometry CursorGeometry = AllottedGeometry.ToPaintGeometry(FVector2D(XPos - CursorHalfLength, 0), FVector2D(2 * CursorHalfLength, AllottedGeometry.GetLocalSize().Y));
 
 		FLinearColor CursorColor = InWidgetStyle.GetColorAndOpacityTint();
@@ -310,14 +310,14 @@ FReply FVisualLoggerTimeSliderController::OnMouseButtonDown( SWidget& WidgetOwne
 	bool bHandleLeftMouseButton = MouseEvent.GetEffectingButton() == EKeys::LeftMouseButton;
 	bool bHandleRightMouseButton = MouseEvent.GetEffectingButton() == EKeys::RightMouseButton && TimeSliderArgs.AllowZoom;
 	
-	DistanceDragged = 0;
+	DistanceDragged = 0.;
 
 	if ( bHandleLeftMouseButton )
 	{
 		// Always capture mouse if we left or right click on the widget
 		FScrubRangeToScreen RangeToScreen(TimeSliderArgs.ViewRange.Get(), MyGeometry.GetLocalSize());
 		FVector2D CursorPos = MyGeometry.AbsoluteToLocal(MouseEvent.GetLastScreenSpacePosition());
-		float NewValue = RangeToScreen.LocalXToInput(CursorPos.X);
+		double NewValue = RangeToScreen.LocalXToInput(CursorPos.X);
 
 		CommitScrubPosition(NewValue, /*bIsScrubbing=*/false);
 		return FReply::Handled().CaptureMouse( WidgetOwner.AsShared() ).PreventThrottling();
@@ -356,7 +356,7 @@ FReply FVisualLoggerTimeSliderController::OnMouseButtonUp( SWidget& WidgetOwner,
 		{
 			FScrubRangeToScreen RangeToScreen( TimeSliderArgs.ViewRange.Get(), MyGeometry.GetLocalSize() );
 			FVector2D CursorPos = MyGeometry.AbsoluteToLocal(MouseEvent.GetLastScreenSpacePosition());
-			float NewValue = RangeToScreen.LocalXToInput(CursorPos.X);
+			double NewValue = RangeToScreen.LocalXToInput(CursorPos.X);
 
 			CommitScrubPosition( NewValue, /*bIsScrubbing=*/false );
 		}
@@ -369,14 +369,14 @@ FReply FVisualLoggerTimeSliderController::OnMouseButtonUp( SWidget& WidgetOwner,
 	return FReply::Unhandled();
 }
 
-float FVisualLoggerTimeSliderController::GetTimeAtCursorPosition(const FGeometry& MyGeometry, const FPointerEvent& MouseEvent) const
+double FVisualLoggerTimeSliderController::GetTimeAtCursorPosition(const FGeometry& MyGeometry, const FPointerEvent& MouseEvent) const
 {
 	FScrubRangeToScreen RangeToScreen(TimeSliderArgs.ViewRange.Get(), MyGeometry.GetLocalSize());
 	FVector2D CursorPos = MyGeometry.AbsoluteToLocal(MouseEvent.GetLastScreenSpacePosition());
-	float NewValue = RangeToScreen.LocalXToInput(CursorPos.X);
+	double NewValue = RangeToScreen.LocalXToInput(CursorPos.X);
 
-	float LocalClampMin = TimeSliderArgs.ClampRange.Get().GetLowerBoundValue();
-	float LocalClampMax = TimeSliderArgs.ClampRange.Get().GetUpperBoundValue();
+	double LocalClampMin = TimeSliderArgs.ClampRange.Get().GetLowerBoundValue();
+	double LocalClampMax = TimeSliderArgs.ClampRange.Get().GetUpperBoundValue();
 
 	if (NewValue < LocalClampMin)
 	{
@@ -411,17 +411,17 @@ FReply FVisualLoggerTimeSliderController::OnMouseMove( SWidget& WidgetOwner, con
 			{
 				SoftwareCursorPosition = MyGeometry.AbsoluteToLocal(MouseEvent.GetLastScreenSpacePosition());
 
-				TRange<float> LocalViewRange = TimeSliderArgs.ViewRange.Get();
-				float LocalViewRangeMin = LocalViewRange.GetLowerBoundValue();
-				float LocalViewRangeMax = LocalViewRange.GetUpperBoundValue();
+				TRange<double> LocalViewRange = TimeSliderArgs.ViewRange.Get();
+				double LocalViewRangeMin = LocalViewRange.GetLowerBoundValue();
+				double LocalViewRangeMax = LocalViewRange.GetUpperBoundValue();
 
 				FScrubRangeToScreen ScaleInfo( LocalViewRange, MyGeometry.GetLocalSize() );
 				FVector2D ScreenDelta = MouseEvent.GetCursorDelta();
 				FVector2D InputDelta;
 				InputDelta.X = ScreenDelta.X/ScaleInfo.PixelsPerInput;
 
-				const TRange<float> NewViewRange = TRange<float>(LocalViewRangeMin - InputDelta.X, LocalViewRangeMax - InputDelta.X);
-				TRange<float> LocalClampRange = TimeSliderArgs.ClampRange.Get();
+				const TRange<double> NewViewRange = TRange<double>(LocalViewRangeMin - InputDelta.X, LocalViewRangeMax - InputDelta.X);
+				TRange<double> LocalClampRange = TimeSliderArgs.ClampRange.Get();
 
 				// Do not try to pan outside the clamp range to prevent undesired zoom
 				if (LocalClampRange.Contains(NewViewRange))
@@ -429,9 +429,9 @@ FReply FVisualLoggerTimeSliderController::OnMouseMove( SWidget& WidgetOwner, con
 					TimeSliderArgs.OnViewRangeChanged.ExecuteIfBound(NewViewRange);
 					if (Scrollbar.IsValid())
 					{
-						float InOffsetFraction = (NewViewRange.GetLowerBoundValue() - LocalClampRange.GetLowerBoundValue()) / LocalClampRange.Size<float>();
-						float InThumbSizeFraction = NewViewRange.Size<float>() / LocalClampRange.Size<float>();
-						Scrollbar->SetState(InOffsetFraction, InThumbSizeFraction);
+						double InOffsetFraction = (NewViewRange.GetLowerBoundValue() - LocalClampRange.GetLowerBoundValue()) / LocalClampRange.Size<double>();
+						double InThumbSizeFraction = NewViewRange.Size<double>() / LocalClampRange.Size<double>();
+						Scrollbar->SetState(static_cast<float>(InOffsetFraction), static_cast<float>(InThumbSizeFraction));
 					}
 
 					if (!TimeSliderArgs.ViewRange.IsBound())
@@ -447,7 +447,7 @@ FReply FVisualLoggerTimeSliderController::OnMouseMove( SWidget& WidgetOwner, con
 			if ( !bDraggingScrubber )
 			{
 				DistanceDragged += FMath::Abs( MouseEvent.GetCursorDelta().X );
-				if ( DistanceDragged > 0/*FSlateApplication::Get().GetDragTriggerDistance()*/ )
+				if ( DistanceDragged > 0. /*FSlateApplication::Get().GetDragTriggerDistance()*/ )
 				{
 					bDraggingScrubber = true;
 					TimeSliderArgs.OnBeginScrubberMovement.ExecuteIfBound();
@@ -455,7 +455,7 @@ FReply FVisualLoggerTimeSliderController::OnMouseMove( SWidget& WidgetOwner, con
 			}
 			else
 			{
-				const float NewValue = GetTimeAtCursorPosition(MyGeometry, MouseEvent);
+				const double NewValue = GetTimeAtCursorPosition(MyGeometry, MouseEvent);
 				CommitScrubPosition(NewValue, /*bIsScrubbing=*/true);
 			}
 		}
@@ -465,7 +465,7 @@ FReply FVisualLoggerTimeSliderController::OnMouseMove( SWidget& WidgetOwner, con
 	return FReply::Unhandled();
 }
 
-void FVisualLoggerTimeSliderController::CommitScrubPosition( float NewValue, bool bIsScrubbing )
+void FVisualLoggerTimeSliderController::CommitScrubPosition( double NewValue, bool bIsScrubbing )
 {
 	// Manage the scrub position ourselves if its not bound to a delegate
 	if ( !TimeSliderArgs.ScrubPosition.IsBound() )
@@ -473,8 +473,8 @@ void FVisualLoggerTimeSliderController::CommitScrubPosition( float NewValue, boo
 		TimeSliderArgs.ScrubPosition.Set( NewValue );
 	}
 
-	TRange<float> LocalViewRange = TimeSliderArgs.ViewRange.Get();
-	const float RangeSize = LocalViewRange.Size<float>();
+	TRange<double> LocalViewRange = TimeSliderArgs.ViewRange.Get();
+	const double RangeSize = LocalViewRange.Size<double>();
 	if (NewValue < LocalViewRange.GetLowerBoundValue())
 	{
 		SetTimeRange(NewValue, NewValue + RangeSize);
@@ -497,54 +497,54 @@ void FVisualLoggerTimeSliderController::HorizontalScrollBar_OnUserScrolled(float
 {
 	if (!TimeSliderArgs.ViewRange.IsBound())
 	{
-		TRange<float> LocalViewRange = TimeSliderArgs.ViewRange.Get();
-		float LocalViewRangeMin = LocalViewRange.GetLowerBoundValue();
-		float LocalViewRangeMax = LocalViewRange.GetUpperBoundValue();
-		float LocalClampMin = TimeSliderArgs.ClampRange.Get().GetLowerBoundValue();
-		float LocalClampMax = TimeSliderArgs.ClampRange.Get().GetUpperBoundValue();
+		TRange<double> LocalViewRange = TimeSliderArgs.ViewRange.Get();
+		double LocalViewRangeMin = LocalViewRange.GetLowerBoundValue();
+		double LocalViewRangeMax = LocalViewRange.GetUpperBoundValue();
+		double LocalClampMin = TimeSliderArgs.ClampRange.Get().GetLowerBoundValue();
+		double LocalClampMax = TimeSliderArgs.ClampRange.Get().GetUpperBoundValue();
 
-		float InThumbSizeFraction = (LocalViewRangeMax - LocalViewRangeMin) / (LocalClampMax - LocalClampMin);
+		double InThumbSizeFraction = (LocalViewRangeMax - LocalViewRangeMin) / (LocalClampMax - LocalClampMin);
 
-		float NewViewOutputMin = LocalClampMin + ScrollOffset * (LocalClampMax - LocalClampMin);
+		double NewViewOutputMin = LocalClampMin + ScrollOffset * (LocalClampMax - LocalClampMin);
 		// The  output is not bound to a delegate so we'll manage the value ourselves
-		float NewViewOutputMax = FMath::Min<float>(NewViewOutputMin + (LocalViewRangeMax - LocalViewRangeMin), LocalClampMax);
+		double NewViewOutputMax = FMath::Min(NewViewOutputMin + (LocalViewRangeMax - LocalViewRangeMin), LocalClampMax);
 		NewViewOutputMin = NewViewOutputMax - (LocalViewRangeMax - LocalViewRangeMin);
 
-		float InOffsetFraction = (NewViewOutputMin - LocalClampMin) / (LocalClampMax - LocalClampMin);
+		double InOffsetFraction = (NewViewOutputMin - LocalClampMin) / (LocalClampMax - LocalClampMin);
 		//if (InOffsetFraction + InThumbSizeFraction <= 1)
 		{
-			TimeSliderArgs.ViewRange.Set(TRange<float>(NewViewOutputMin, NewViewOutputMax));
-			Scrollbar->SetState(InOffsetFraction, InThumbSizeFraction);
+			TimeSliderArgs.ViewRange.Set(TRange<double>(NewViewOutputMin, NewViewOutputMax));
+			Scrollbar->SetState(static_cast<float>(InOffsetFraction), static_cast<float>(InThumbSizeFraction));
 		}
 	}
 }
 
-void FVisualLoggerTimeSliderController::SetTimeRange(float NewViewOutputMin, float NewViewOutputMax)
+void FVisualLoggerTimeSliderController::SetTimeRange(double NewViewOutputMin, double NewViewOutputMax)
 {
-	TimeSliderArgs.ViewRange.Set(TRange<float>(NewViewOutputMin, NewViewOutputMax));
+	TimeSliderArgs.ViewRange.Set(TRange<double>(NewViewOutputMin, NewViewOutputMax));
 
-	float LocalClampMin = TimeSliderArgs.ClampRange.Get().GetLowerBoundValue();
-	float LocalClampMax = TimeSliderArgs.ClampRange.Get().GetUpperBoundValue();
+	double LocalClampMin = TimeSliderArgs.ClampRange.Get().GetLowerBoundValue();
+	double LocalClampMax = TimeSliderArgs.ClampRange.Get().GetUpperBoundValue();
 
-	const float InOffsetFraction = (NewViewOutputMin - LocalClampMin) / (LocalClampMax - LocalClampMin);
-	const float InThumbSizeFraction = (NewViewOutputMax - NewViewOutputMin) / (LocalClampMax - LocalClampMin);
-	Scrollbar->SetState(InOffsetFraction, InThumbSizeFraction);
+	const double InOffsetFraction = (NewViewOutputMin - LocalClampMin) / (LocalClampMax - LocalClampMin);
+	const double InThumbSizeFraction = (NewViewOutputMax - NewViewOutputMin) / (LocalClampMax - LocalClampMin);
+	Scrollbar->SetState(static_cast<float>(InOffsetFraction), static_cast<float>(InThumbSizeFraction));
 }
 
-void FVisualLoggerTimeSliderController::SetClampRange(float MinValue, float MaxValue)
+void FVisualLoggerTimeSliderController::SetClampRange(double MinValue, double MaxValue)
 {
-	TRange<float> LocalViewRange = TimeSliderArgs.ViewRange.Get();
-	float LocalClampMin = TimeSliderArgs.ClampRange.Get().GetLowerBoundValue();
-	float LocalClampMax = TimeSliderArgs.ClampRange.Get().GetUpperBoundValue();
-	const float CurrentDistance = LocalClampMax - LocalClampMin;
-	const float ZoomDelta = (LocalViewRange.GetUpperBoundValue() - LocalViewRange.GetLowerBoundValue()) / CurrentDistance;
+	TRange<double> LocalViewRange = TimeSliderArgs.ViewRange.Get();
+	double LocalClampMin = TimeSliderArgs.ClampRange.Get().GetLowerBoundValue();
+	double LocalClampMax = TimeSliderArgs.ClampRange.Get().GetUpperBoundValue();
+	const double CurrentDistance = LocalClampMax - LocalClampMin;
+	const double ZoomDelta = (LocalViewRange.GetUpperBoundValue() - LocalViewRange.GetLowerBoundValue()) / CurrentDistance;
 
 	MaxValue = MinValue + (MaxValue - MinValue < 2 ? CurrentDistance : MaxValue - MinValue);
 
-	TimeSliderArgs.ClampRange = TRange<float>(MinValue, MaxValue);
+	TimeSliderArgs.ClampRange = TRange<double>(MinValue, MaxValue);
 
-	const float LocalViewRangeMin = FMath::Clamp(LocalViewRange.GetLowerBoundValue(), MinValue, MaxValue);
-	const float LocalViewRangeMax = FMath::Clamp(LocalViewRange.GetUpperBoundValue(), MinValue, MaxValue);
+	const double LocalViewRangeMin = FMath::Clamp(LocalViewRange.GetLowerBoundValue(), MinValue, MaxValue);
+	const double LocalViewRangeMax = FMath::Clamp(LocalViewRange.GetUpperBoundValue(), MinValue, MaxValue);
 	SetTimeRange(ZoomDelta >= 1 ? MinValue : LocalViewRangeMin, ZoomDelta >= 1 ? MaxValue : LocalViewRangeMax);
 }
 
@@ -564,21 +564,21 @@ FReply FVisualLoggerTimeSliderController::OnMouseWheel( SWidget& WidgetOwner, co
 		const float ZoomDelta = -0.1f * MouseEvent.GetWheelDelta();
 
 		{
-			float MouseFractionX = MyGeometry.AbsoluteToLocal(MouseEvent.GetScreenSpacePosition()).X / MyGeometry.GetLocalSize().X;
+			double MouseFractionX = MyGeometry.AbsoluteToLocal(MouseEvent.GetScreenSpacePosition()).X / MyGeometry.GetLocalSize().X;
 
-			TRange<float> LocalViewRange = TimeSliderArgs.ViewRange.Get();
-			float LocalViewRangeMax = LocalViewRange.GetUpperBoundValue();
-			float LocalViewRangeMin = LocalViewRange.GetLowerBoundValue();
-			const float OutputViewSize = LocalViewRangeMax - LocalViewRangeMin;
-			const float OutputChange = OutputViewSize * ZoomDelta;
+			TRange<double> LocalViewRange = TimeSliderArgs.ViewRange.Get();
+			double LocalViewRangeMax = LocalViewRange.GetUpperBoundValue();
+			double LocalViewRangeMin = LocalViewRange.GetLowerBoundValue();
+			const double OutputViewSize = LocalViewRangeMax - LocalViewRangeMin;
+			const double OutputChange = OutputViewSize * ZoomDelta;
 
-			float NewViewOutputMin = LocalViewRangeMin - (OutputChange * MouseFractionX);
-			float NewViewOutputMax = LocalViewRangeMax + (OutputChange * (1.0f - MouseFractionX));
+			double NewViewOutputMin = LocalViewRangeMin - (OutputChange * MouseFractionX);
+			double NewViewOutputMax = LocalViewRangeMax + (OutputChange * (1.0f - MouseFractionX));
 
 			if (NewViewOutputMin < NewViewOutputMax)
 			{
-				float LocalClampMin = TimeSliderArgs.ClampRange.Get().GetLowerBoundValue();
-				float LocalClampMax = TimeSliderArgs.ClampRange.Get().GetUpperBoundValue();
+				double LocalClampMin = TimeSliderArgs.ClampRange.Get().GetLowerBoundValue();
+				double LocalClampMax = TimeSliderArgs.ClampRange.Get().GetUpperBoundValue();
 
 				// Clamp the range if clamp values are set
 				if ( NewViewOutputMin < LocalClampMin )
@@ -591,17 +591,17 @@ FReply FVisualLoggerTimeSliderController::OnMouseWheel( SWidget& WidgetOwner, co
 					NewViewOutputMax = LocalClampMax;
 				}
 
-				TimeSliderArgs.OnViewRangeChanged.ExecuteIfBound(TRange<float>(NewViewOutputMin, NewViewOutputMax));
+				TimeSliderArgs.OnViewRangeChanged.ExecuteIfBound(TRange<double>(NewViewOutputMin, NewViewOutputMax));
 				if (Scrollbar.IsValid())
 				{
-					float InOffsetFraction = (NewViewOutputMin - LocalClampMin) / (LocalClampMax - LocalClampMin);
-					float InThumbSizeFraction = (NewViewOutputMax - NewViewOutputMin) / (LocalClampMax - LocalClampMin);
-					Scrollbar->SetState(InOffsetFraction, InThumbSizeFraction);
+					const double InOffsetFraction = (NewViewOutputMin - LocalClampMin) / (LocalClampMax - LocalClampMin);
+					const double InThumbSizeFraction = (NewViewOutputMax - NewViewOutputMin) / (LocalClampMax - LocalClampMin);
+					Scrollbar->SetState(static_cast<float>(InOffsetFraction), static_cast<float>(InThumbSizeFraction));
 				}
 				if( !TimeSliderArgs.ViewRange.IsBound() )
 				{	
 					// The  output is not bound to a delegate so we'll manage the value ourselves
-					TimeSliderArgs.ViewRange.Set( TRange<float>( NewViewOutputMin, NewViewOutputMax ) );
+					TimeSliderArgs.ViewRange.Set( TRange<double>( NewViewOutputMin, NewViewOutputMax ) );
 				}
 			}
 		}
@@ -616,12 +616,12 @@ int32 FVisualLoggerTimeSliderController::OnPaintSectionView( const FGeometry& Al
 {
 	const ESlateDrawEffect DrawEffects = bEnabled ? ESlateDrawEffect::None : ESlateDrawEffect::DisabledEffect;
 
-	TRange<float> LocalViewRange = TimeSliderArgs.ViewRange.Get();
-	float LocalScrubPosition = TimeSliderArgs.ScrubPosition.Get();
+	TRange<double> LocalViewRange = TimeSliderArgs.ViewRange.Get();
+	double LocalScrubPosition = TimeSliderArgs.ScrubPosition.Get();
 
-	float ViewRange = LocalViewRange.Size<float>();
-	float PixelsPerInput = ViewRange > 0 ? AllottedGeometry.GetLocalSize().X / ViewRange : 0;
-	float LinePos =  (LocalScrubPosition - LocalViewRange.GetLowerBoundValue()) * PixelsPerInput;
+	double ViewRange = LocalViewRange.Size<double>();
+	double PixelsPerInput = ViewRange > 0. ? AllottedGeometry.GetLocalSize().X / ViewRange : 0.;
+	double LinePos =  (LocalScrubPosition - LocalViewRange.GetLowerBoundValue()) * PixelsPerInput;
 
 	FScrubRangeToScreen RangeToScreen( LocalViewRange, AllottedGeometry.GetLocalSize());
 
@@ -638,7 +638,7 @@ int32 FVisualLoggerTimeSliderController::OnPaintSectionView( const FGeometry& Al
 		// Draw major ticks under sections
 		Args.StartLayer = LayerId-1;
 		// Draw the tick the entire height of the section area
-		Args.TickOffset = 0.0f;
+		Args.TickOffset = 0.;
 		Args.MajorTickHeight = AllottedGeometry.GetLocalSize().Y;
 
 		DrawTicks( OutDrawElements, RangeToScreen, Args );
@@ -648,8 +648,8 @@ int32 FVisualLoggerTimeSliderController::OnPaintSectionView( const FGeometry& Al
 	{
 		// Draw cursor size
 		const float CursorHalfSize = TimeSliderArgs.CursorSize.Get() * 0.5f;
-		const float CursorHalfLength = AllottedGeometry.GetLocalSize().X * CursorHalfSize;
-		FPaintGeometry CursorGeometry = AllottedGeometry.ToPaintGeometry(FVector2D(LinePos - CursorHalfLength, 0), FVector2D(2 * CursorHalfLength, AllottedGeometry.GetLocalSize().Y));
+		const double CursorHalfLength = AllottedGeometry.GetLocalSize().X * CursorHalfSize;
+		FPaintGeometry CursorGeometry = AllottedGeometry.ToPaintGeometry(FVector2D(LinePos - CursorHalfLength, 0.), FVector2D(2 * CursorHalfLength, AllottedGeometry.GetLocalSize().Y));
 
 		FSlateDrawElement::MakeBox(
 			OutDrawElements,
@@ -663,13 +663,13 @@ int32 FVisualLoggerTimeSliderController::OnPaintSectionView( const FGeometry& Al
 		// Draw a line for the scrub position
 		TArray<FVector2D> LinePoints;
 		LinePoints.AddUninitialized(2);
-		LinePoints[0] = FVector2D( 1.0f, 0.0f );
-		LinePoints[1] = FVector2D( 1.0f, FMath::RoundToFloat( AllottedGeometry.GetLocalSize().Y ) );
+		LinePoints[0] = FVector2D( 1., 0. );
+		LinePoints[1] = FVector2D( 1., FMath::RoundToFloat( AllottedGeometry.GetLocalSize().Y ) );
 
 		FSlateDrawElement::MakeLines(
 			OutDrawElements,
 			++LayerId,
-			AllottedGeometry.ToPaintGeometry( FVector2D(LinePos, 0.0f ), FVector2D(1.0f,1.0f) ),
+			AllottedGeometry.ToPaintGeometry( FVector2D(LinePos, 0. ), FVector2D(1.,1.) ),
 			LinePoints,
 			DrawEffects,
 			FLinearColor::White.CopyWithNewOpacity(0.39f),
diff --git a/Engine/Source/Developer/LogVisualizer/Private/VisualLoggerTimeSliderController.h b/Engine/Source/Developer/LogVisualizer/Private/VisualLoggerTimeSliderController.h
--- a/Engine/Source/Developer/LogVisualizer/Private/VisualLoggerTimeSliderController.h
+++ b/Engine/Source/Developer/LogVisualizer/Private/VisualLoggerTimeSliderController.h
@@ -13,19 +13,19 @@ class FSlateWindowElementList;
 
 struct FVisualLoggerTimeSliderArgs
 {
-	DECLARE_DELEGATE_OneParam( FOnRangeChanged, TRange<float> )
-	DECLARE_DELEGATE_TwoParams( FOnScrubPositionChanged, float, bool )
+	DECLARE_DELEGATE_OneParam( FOnRangeChanged, TRange<double> )
+	DECLARE_DELEGATE_TwoParams( FOnScrubPositionChanged, double, bool )
 
 	FVisualLoggerTimeSliderArgs()
-		: ScrubPosition(0)
-		, ViewRange(TRange<float>(0.0f, 5.0f))
-		, ClampRange(TRange<float>(0.0f, 5.0f))
+		: ScrubPosition(0.)
+		, ViewRange(TRange<double>(0., 5.))
+		, ClampRange(TRange<double>(0., 5.))
 		, AllowZoom(true)
 		, CursorSize(0.05f)
 	{}
 
 	/** The scrub position */
-	TAttribute<float> ScrubPosition;
+	TAttribute<double> ScrubPosition;
 
 	/** Called when the scrub position changes */
 	FOnScrubPositionChanged OnScrubPositionChanged;
@@ -37,16 +37,16 @@ struct FVisualLoggerTimeSliderArgs
 	FSimpleDelegate OnEndScrubberMovement;
 
 	/** View time range */
-	TAttribute< TRange<float> > ViewRange;
+	TAttribute< TRange<double> > ViewRange;
 
 	/** Clamp time range */
-	TAttribute< TRange<float> > ClampRange;
+	TAttribute< TRange<double> > ClampRange;
 
 	/** Called when the view range changes */
 	FOnRangeChanged OnViewRangeChanged;
 
 	/** Attribute defining the time snap interval */
-	TAttribute<float> TimeSnapInterval;
+	TAttribute<double> TimeSnapInterval;
 
 	/** If we are allowed to zoom */
 	bool AllowZoom;
@@ -73,7 +73,7 @@ public:
 	* @param MinTickSpacing	The minimum tick spacing in time units allowed
 	* @return the optimal spacing in time units
 	*/
-	float DetermineOptimalSpacing(float InPixelsPerInput, uint32 MinTick, float MinTickSpacing) const;
+	double DetermineOptimalSpacing(double InPixelsPerInput, uint32 MinTick, float MinTickSpacing) const;
 	void SetTimesliderArgs(const FVisualLoggerTimeSliderArgs& InArgs);
 
 	int32 OnPaintTimeSlider( bool bMirrorLabels, const FGeometry& AllottedGeometry, const FSlateRect& MyCullingRect, FSlateWindowElementList& OutDrawElements, int32 LayerId, const FWidgetStyle& InWidgetStyle, bool bParentEnabled ) const;
@@ -94,14 +94,14 @@ public:
 	 * @param NewValue				Value resulting from the user's interaction
 	 * @param bIsScrubbing			True if done via scrubbing, false if just releasing scrubbing
 	 */
-	void CommitScrubPosition( float NewValue, bool bIsScrubbing );
+	void CommitScrubPosition(double NewValue, bool bIsScrubbing );
 
 	void SetExternalScrollbar(TSharedRef<SScrollBar> Scrollbar);
-	void SetTimeRange(float MinValue, float MaxValue);
-	void SetClampRange(float MinValue, float MaxValue);
+	void SetTimeRange(double MinValue, double MaxValue);
+	void SetClampRange(double MinValue, double MaxValue);
 	bool IsPanning() { return bPanning; }
 
-	float GetTimeAtCursorPosition(const FGeometry& MyGeometry, const FPointerEvent& MouseEvent) const;
+	double GetTimeAtCursorPosition(const FGeometry& MyGeometry, const FPointerEvent& MouseEvent) const;
 private:
 	// forward declared as class members to prevent name collision with similar types defined in other units
 	struct FScrubRangeToScreen;
@@ -125,7 +125,7 @@ private:
 	/** Brush for drawing cursor background to visualize corsor size */
 	const FSlateBrush* CursorBackground;
 	/** Total mouse delta during dragging **/
-	float DistanceDragged;
+	double DistanceDragged;
 	/** If we are dragging the scrubber */
 	bool bDraggingScrubber;
 	/** If we are currently panning the panel */
diff --git a/Engine/Source/Developer/LogVisualizer/Public/LogVisualizerPublic.h b/Engine/Source/Developer/LogVisualizer/Public/LogVisualizerPublic.h
--- a/Engine/Source/Developer/LogVisualizer/Public/LogVisualizerPublic.h
+++ b/Engine/Source/Developer/LogVisualizer/Public/LogVisualizerPublic.h
@@ -56,7 +56,9 @@ struct LOGVISUALIZER_API FLogVisualizer
 	int32 GetNextItem(FName RowName, int32 MoveDistance = 1);
 	int32 GetPreviousItem(FName RowName, int32 MoveDistance = 1);
 
+	// @todo: This function currently doesn't do anything!
 	void GotoNextItem(FName RowName, int32 MoveDistance = 1);
+	// @todo: This function currently doesn't do anything!
 	void GotoPreviousItem(FName RowName, int32 MoveDistance = 1);
 	void GotoFirstItem(FName RowName);
 	void GotoLastItem(FName RowName);
diff --git a/Engine/Source/Runtime/AIModule/Classes/VisualLoggerExtension.h b/Engine/Source/Runtime/AIModule/Classes/VisualLoggerExtension.h
--- a/Engine/Source/Runtime/AIModule/Classes/VisualLoggerExtension.h
+++ b/Engine/Source/Runtime/AIModule/Classes/VisualLoggerExtension.h
@@ -32,7 +32,7 @@ public:
 	virtual void OnLogLineSelectionChanged(IVisualLoggerEditorInterface* EdInterface, TSharedPtr<struct FLogEntryItem> SelectedItem, int64 UserData) override;
 
 private:
-	void DrawData(UWorld* InWorld, class UEQSRenderingComponent* EQSRenderingComponent, UCanvas* Canvas, AActor* HelperActor, const FName& TagName, const FVisualLogDataBlock& DataBlock, float Timestamp);
+	void DrawData(UWorld* InWorld, class UEQSRenderingComponent* EQSRenderingComponent, UCanvas* Canvas, AActor* HelperActor, const FName& TagName, const FVisualLogDataBlock& DataBlock, double Timestamp);
 	void DisableEQSRendering(AActor* HelperActor);
 
 protected:
diff --git a/Engine/Source/Runtime/AIModule/Private/VisualLoggerExtension.cpp b/Engine/Source/Runtime/AIModule/Private/VisualLoggerExtension.cpp
--- a/Engine/Source/Runtime/AIModule/Private/VisualLoggerExtension.cpp
+++ b/Engine/Source/Runtime/AIModule/Private/VisualLoggerExtension.cpp
@@ -112,7 +112,7 @@ void FVisualLoggerExtension::DrawData(IVisualLoggerEditorInterface* EdInterface,
 	}
 }
 
-void FVisualLoggerExtension::DrawData(UWorld* InWorld, class UEQSRenderingComponent* EQSRenderComp, UCanvas* Canvas, AActor* HelperActor, const FName& TagName, const FVisualLogDataBlock& DataBlock, float Timestamp)
+void FVisualLoggerExtension::DrawData(UWorld* InWorld, class UEQSRenderingComponent* EQSRenderComp, UCanvas* Canvas, AActor* HelperActor, const FName& TagName, const FVisualLogDataBlock& DataBlock, double Timestamp)
 {
 #if USE_EQS_DEBUGGER
 	if (TagName == *EVisLogTags::TAG_EQS)
diff --git a/Engine/Source/Runtime/Core/Public/Math/Vector.h b/Engine/Source/Runtime/Core/Public/Math/Vector.h
--- a/Engine/Source/Runtime/Core/Public/Math/Vector.h
+++ b/Engine/Source/Runtime/Core/Public/Math/Vector.h
@@ -2669,6 +2669,16 @@ FORCEINLINE FVector3f NarrowWorldPositionChecked(const FVector& WorldPosition)
 	return Narrowed;
 }
 
+// Validated narrowing cast for world positions. FVector -> FVector3f
+FORCEINLINE FVector3f NarrowWorldPositionChecked(const FVector::FReal InX, const FVector::FReal InY, const FVector::FReal InZ)
+{
+	FVector3f Narrowed;
+	Narrowed.X = FloatCastChecked<float>(InX, DefaultFloatPrecision);
+	Narrowed.Y = FloatCastChecked<float>(InY, DefaultFloatPrecision);
+	Narrowed.Z = FloatCastChecked<float>(InZ, DefaultFloatPrecision);
+	return Narrowed;
+}
+
 } // namespace UE::LWC
 
 } // namespace UE
diff --git a/Engine/Source/Runtime/Core/Public/UObject/UE5MainStreamObjectVersion.h b/Engine/Source/Runtime/Core/Public/UObject/UE5MainStreamObjectVersion.h
--- a/Engine/Source/Runtime/Core/Public/UObject/UE5MainStreamObjectVersion.h
+++ b/Engine/Source/Runtime/Core/Public/UObject/UE5MainStreamObjectVersion.h
@@ -306,6 +306,9 @@ struct CORE_API FUE5MainStreamObjectVersion
 
 		// Store the RigVM execute context struct the VM uses in the archive
 		RigVMSerializeExecuteContextStruct,
+
+		// Store the Visual Logger timestamp as a double
+		VisualLoggerTimeStampAsDouble,
 		
 		// -----<new versions can be added above this line>-------------------------------------------------
 		VersionPlusOne,
diff --git a/Engine/Source/Runtime/Engine/Classes/VisualLogger/VisualLoggerBinaryFileDevice.h b/Engine/Source/Runtime/Engine/Classes/VisualLogger/VisualLoggerBinaryFileDevice.h
--- a/Engine/Source/Runtime/Engine/Classes/VisualLogger/VisualLoggerBinaryFileDevice.h
+++ b/Engine/Source/Runtime/Engine/Classes/VisualLogger/VisualLoggerBinaryFileDevice.h
@@ -20,8 +20,8 @@ public:
 
 	FVisualLoggerBinaryFileDevice();
 	virtual void Cleanup(bool bReleaseMemory = false) override;
-	virtual void StartRecordingToFile(float TImeStamp) override;
-	virtual void StopRecordingToFile(float TImeStamp) override;
+	virtual void StartRecordingToFile(double TimeStamp) override;
+	virtual void StopRecordingToFile(double TimeStamp) override;
 	virtual void DiscardRecordingToFile() override;
 	virtual void SetFileName(const FString& InFileName) override;
 	virtual void Serialize(const class UObject* LogOwner, FName OwnerName, FName OwnerClassName, const FVisualLogEntry& LogEntry) override;
@@ -31,8 +31,8 @@ public:
 protected:
 	int32 bUseCompression : 1;
 	float FrameCacheLenght;
-	float StartRecordingTime;
-	float LastLogTimeStamp;
+	double StartRecordingTime;
+	double LastLogTimeStamp;
 	FArchive* FileArchive;
 	FString TempFileName;
 	FString FileName;
diff --git a/Engine/Source/Runtime/Engine/Classes/VisualLogger/VisualLoggerTraceDevice.h b/Engine/Source/Runtime/Engine/Classes/VisualLogger/VisualLoggerTraceDevice.h
--- a/Engine/Source/Runtime/Engine/Classes/VisualLogger/VisualLoggerTraceDevice.h
+++ b/Engine/Source/Runtime/Engine/Classes/VisualLogger/VisualLoggerTraceDevice.h
@@ -18,8 +18,8 @@ public:
 
 	FVisualLoggerTraceDevice();
 	virtual void Cleanup(bool bReleaseMemory = false) override;
-	virtual void StartRecordingToFile(float TImeStamp) override;
-	virtual void StopRecordingToFile(float TImeStamp) override;
+	virtual void StartRecordingToFile(double TimeStamp) override;
+	virtual void StopRecordingToFile(double TimeStamp) override;
 	virtual void DiscardRecordingToFile() override;
 	virtual void SetFileName(const FString& InFileName) override;
 	virtual void Serialize(const class UObject* LogOwner, FName OwnerName, FName OwnerClassName, const FVisualLogEntry& LogEntry) override;
diff --git a/Engine/Source/Runtime/Engine/Private/VisualLogger/VisualLogger.cpp b/Engine/Source/Runtime/Engine/Private/VisualLogger/VisualLogger.cpp
--- a/Engine/Source/Runtime/Engine/Private/VisualLogger/VisualLogger.cpp
+++ b/Engine/Source/Runtime/Engine/Private/VisualLogger/VisualLogger.cpp
@@ -302,7 +302,7 @@ bool FVisualLogger::CheckVisualLogInputInternal(const UObject* Object, const FNa
 	return true;
 }
 
-float FVisualLogger::GetTimeStampForObject(const UObject* Object) const
+double FVisualLogger::GetTimeStampForObject(const UObject* Object) const
 {
 	if (GetTimeStampFunc)
 	{
@@ -314,10 +314,10 @@ float FVisualLogger::GetTimeStampForObject(const UObject* Object) const
 		return World->TimeSeconds;
 	}
 
-	return 0.0f;
+	return 0.;
 }
 
-void FVisualLogger::SetGetTimeStampFunc(const TFunction<float(const UObject*)> Function)
+void FVisualLogger::SetGetTimeStampFunc(const TFunction<double(const UObject*)> Function)
 {
 	GetTimeStampFunc = Function;
 }
@@ -416,7 +416,7 @@ FVisualLogEntry* FVisualLogger::GetLastEntryForObject(const UObject* Object)
 	return ThreadCurrentEntryPerObject.Find(LogOwner);
 }
 
-FVisualLogEntry* FVisualLogger::GetEntryToWrite(const UObject* Object, const float TimeStamp, ECreateIfNeeded ShouldCreate)
+FVisualLogEntry* FVisualLogger::GetEntryToWrite(const UObject* Object, const double TimeStamp, ECreateIfNeeded ShouldCreate)
 {
 	const UObject* LogOwner = nullptr;
 	{
@@ -489,7 +489,7 @@ FVisualLogEntry* FVisualLogger::GetEntryToWrite(const UObject* Object, const flo
 }
 
 
-FVisualLogEntry* FVisualLogger::GetEntryToWriteInternal(const UObject* Object, const float TimeStamp, const ECreateIfNeeded ShouldCreate)
+FVisualLogEntry* FVisualLogger::GetEntryToWriteInternal(const UObject* Object, const double TimeStamp, const ECreateIfNeeded ShouldCreate)
 {
 	// No redirection needed, it should have been done at the time of the thread entry was computed
 	const UObject* LogOwner = Object;
@@ -882,7 +882,7 @@ void FVisualLogger::Cleanup(UWorld* OldWorld, const bool bReleaseMemory)
 	}
 }
 
-int32 FVisualLogger::GetUniqueId(const float Timestamp)
+int32 FVisualLogger::GetUniqueId(const double Timestamp)
 {
 	return LastUniqueIds.FindOrAdd(Timestamp)++;
 }
diff --git a/Engine/Source/Runtime/Engine/Private/VisualLogger/VisualLoggerAutomationTests.cpp b/Engine/Source/Runtime/Engine/Private/VisualLogger/VisualLoggerAutomationTests.cpp
--- a/Engine/Source/Runtime/Engine/Private/VisualLogger/VisualLoggerAutomationTests.cpp
+++ b/Engine/Source/Runtime/Engine/Private/VisualLogger/VisualLoggerAutomationTests.cpp
@@ -134,7 +134,7 @@ bool FVisualLogTest::RunTest(const FString& Parameters)
 
 	{
 		const FString TextToLog = TEXT("Simple text line to test if UE_VLOG_UELOG works fine");
-		float CurrentTimestamp = World->TimeSeconds;
+		double CurrentTimestamp = World->TimeSeconds;
 		UE_VLOG_UELOG(World, LogVisual, Log, TEXT("%s"), *TextToLog);
 		CHECK_SUCCESS(Context.Device.LastObject != World);
 		CHECK_SUCCESS(Context.Device.LastEntry.TimeStamp == -1);
@@ -147,7 +147,7 @@ bool FVisualLogTest::RunTest(const FString& Parameters)
 			CHECK_SUCCESS(CurrentEntry->LogLines[0].Category == LogVisual.GetCategoryName());
 			CHECK_SUCCESS(CurrentEntry->LogLines[0].Line == TextToLog);
 
-			const float NewTimestamp = CurrentTimestamp + 0.1f;
+			const double NewTimestamp = CurrentTimestamp + 0.1;
 			FVisualLogEntry* NewEntry = FVisualLogger::Get().GetEntryToWrite(World, NewTimestamp); //generate new entry and serialize old one
 			FVisualLogger::Get().FlushThreadsEntries();
 			CurrentEntry = &Context.Device.LastEntry;
@@ -194,7 +194,7 @@ bool FVisualLogSegmentsTest::RunTest(const FString& Parameters)
 		CHECK_SUCCESS(Context.Device.LastEntry.TimeStamp == -1);
 		FVisualLogEntry* CurrentEntry = FVisualLogger::Get().GetEntryToWrite(World, World->TimeSeconds, ECreateIfNeeded::DontCreate);
 
-		float CurrentTimestamp = World->TimeSeconds;
+		double CurrentTimestamp = World->TimeSeconds;
 		{
 			CHECK_NOT_NULL(CurrentEntry);
 			CHECK_SUCCESS(CurrentEntry->TimeStamp == CurrentTimestamp);
@@ -204,7 +204,7 @@ bool FVisualLogSegmentsTest::RunTest(const FString& Parameters)
 			CHECK_SUCCESS(CurrentEntry->ElementsToDraw[0].Points[0] == StartPoint);
 			CHECK_SUCCESS(CurrentEntry->ElementsToDraw[0].Points[1] == EndPoint);
 
-			const float NewTimestamp = CurrentTimestamp + 0.1f;
+			const double NewTimestamp = CurrentTimestamp + 0.1;
 			FVisualLogEntry* NewEntry = FVisualLogger::Get().GetEntryToWrite(World, NewTimestamp); //generate new entry and serialize old one
 			FVisualLogger::Get().FlushThreadsEntries();
 			CurrentEntry = &Context.Device.LastEntry;
@@ -248,7 +248,7 @@ bool FVisualLogEventsTest::RunTest(const FString& Parameters)
 	CHECK_SUCCESS(EventTest3.Name == TEXT("EventTest3"));
 	CHECK_SUCCESS(EventTest3.FriendlyDesc == TEXT("Third simple event for vlog tests"));
 
-	float CurrentTimestamp = World->TimeSeconds;
+	double CurrentTimestamp = World->TimeSeconds;
 	FVisualLogEntry* CurrentEntry = FVisualLogger::Get().GetEntryToWrite(World, CurrentTimestamp, ECreateIfNeeded::DontCreate);
 	CHECK_SUCCESS(CurrentEntry == nullptr);
 
@@ -286,7 +286,7 @@ bool FVisualLogEventsTest::RunTest(const FString& Parameters)
 		CHECK_SUCCESS(CurrentEntry->Events[1].UserFriendlyDesc == TEXT("Second simple event for vlog tests"));
 		CHECK_SUCCESS(CurrentEntry->Events[2].UserFriendlyDesc == TEXT("Third simple event for vlog tests"));
 
-		const float NewTimestamp = CurrentTimestamp + 0.1f;
+		const double NewTimestamp = CurrentTimestamp + 0.1;
 		FVisualLogEntry* NewEntry = FVisualLogger::Get().GetEntryToWrite(World, NewTimestamp); //generate new entry and serialize old one
 		FVisualLogger::Get().FlushThreadsEntries();
 		CurrentEntry = &Context.Device.LastEntry;
@@ -313,7 +313,7 @@ bool FVisualLogEventsTest::RunTest(const FString& Parameters)
 	const FName EventTag2 = TEXT("ATLAS_C_1");
 	const FName EventTag3 = TEXT("ATLAS_C_2");
 
-	CurrentTimestamp = World->TimeSeconds + 0.2f;
+	CurrentTimestamp = World->TimeSeconds + 0.2;
 	CurrentEntry = FVisualLogger::Get().GetEntryToWrite(World, CurrentTimestamp); //generate new entry and serialize old one
 	UE_VLOG_EVENT_WITH_DATA(World, EventTest, EventTag1);
 	CHECK_NOT_NULL(CurrentEntry);
@@ -323,7 +323,7 @@ bool FVisualLogEventsTest::RunTest(const FString& Parameters)
 	CHECK_SUCCESS(CurrentEntry->Events[0].EventTags.Num() == 1);
 	CHECK_SUCCESS(CurrentEntry->Events[0].EventTags[EventTag1] == 1);
 
-	CurrentTimestamp = World->TimeSeconds + 0.3f;
+	CurrentTimestamp = World->TimeSeconds + 0.3;
 	CurrentEntry = FVisualLogger::Get().GetEntryToWrite(World, CurrentTimestamp); //generate new entry and serialize old one
 	UE_VLOG_EVENT_WITH_DATA(World, EventTest, EventTag1, EventTag2, EventTag3);
 	UE_VLOG_EVENT_WITH_DATA(World, EventTest, EventTag3);
diff --git a/Engine/Source/Runtime/Engine/Private/VisualLogger/VisualLoggerBinaryFileDevice.cpp b/Engine/Source/Runtime/Engine/Private/VisualLogger/VisualLoggerBinaryFileDevice.cpp
--- a/Engine/Source/Runtime/Engine/Private/VisualLogger/VisualLoggerBinaryFileDevice.cpp
+++ b/Engine/Source/Runtime/Engine/Private/VisualLogger/VisualLoggerBinaryFileDevice.cpp
@@ -26,7 +26,7 @@ void FVisualLoggerBinaryFileDevice::Cleanup(bool bReleaseMemory)
 
 }
 
-void FVisualLoggerBinaryFileDevice::StartRecordingToFile(float TimeStamp)
+void FVisualLoggerBinaryFileDevice::StartRecordingToFile(double TimeStamp)
 {
 	if (FileArchive != nullptr)
 	{
@@ -44,7 +44,7 @@ void FVisualLoggerBinaryFileDevice::StartRecordingToFile(float TimeStamp)
 	FileArchive = IFileManager::Get().CreateFileWriter(*FullFilename);
 }
 
-void FVisualLoggerBinaryFileDevice::StopRecordingToFile(float TimeStamp)
+void FVisualLoggerBinaryFileDevice::StopRecordingToFile(double TimeStamp)
 {
 	if (FileArchive == nullptr)
 	{
diff --git a/Engine/Source/Runtime/Engine/Private/VisualLogger/VisualLoggerTraceDevice.cpp b/Engine/Source/Runtime/Engine/Private/VisualLogger/VisualLoggerTraceDevice.cpp
--- a/Engine/Source/Runtime/Engine/Private/VisualLogger/VisualLoggerTraceDevice.cpp
+++ b/Engine/Source/Runtime/Engine/Private/VisualLogger/VisualLoggerTraceDevice.cpp
@@ -26,14 +26,14 @@ void FVisualLoggerTraceDevice::Cleanup(bool bReleaseMemory)
 
 }
 
-void FVisualLoggerTraceDevice::StartRecordingToFile(float TimeStamp)
+void FVisualLoggerTraceDevice::StartRecordingToFile(double TimeStamp)
 {
 #if UE_TRACE_ENABLED
 	UE::Trace::ToggleChannel(TEXT("VisualLogger"), true); 
 #endif
 }
 
-void FVisualLoggerTraceDevice::StopRecordingToFile(float TimeStamp)
+void FVisualLoggerTraceDevice::StopRecordingToFile(double TimeStamp)
 {
 #if UE_TRACE_ENABLED
 	UE::Trace::ToggleChannel(TEXT("VisualLogger"), false); 
diff --git a/Engine/Source/Runtime/Engine/Private/VisualLogger/VisualLoggerTypes.cpp b/Engine/Source/Runtime/Engine/Private/VisualLogger/VisualLoggerTypes.cpp
--- a/Engine/Source/Runtime/Engine/Private/VisualLogger/VisualLoggerTypes.cpp
+++ b/Engine/Source/Runtime/Engine/Private/VisualLogger/VisualLoggerTypes.cpp
@@ -3,6 +3,7 @@
 #include "VisualLogger/VisualLoggerTypes.h"
 #include "Engine/World.h"
 #include "Misc/Paths.h"
+#include "UObject/UE5MainStreamObjectVersion.h"
 #include "VisualLogger/VisualLoggerDebugSnapshotInterface.h"
 
 namespace
@@ -18,7 +19,9 @@ UVisualLoggerDebugSnapshotInterface::UVisualLoggerDebugSnapshotInterface(const F
 #if ENABLE_VISUAL_LOG
 
 #define DEPRECATED_VISUAL_LOGGER_MAGIC_NUMBER 0xFAFAAFAF
-#define VISUAL_LOGGER_MAGIC_NUMBER 0xAFAFFAFA
+#define VISUAL_LOGGER_MAGIC_NUMBER_OLD_CUSTOM_VERSION 0xAFAFFAFA
+#define VISUAL_LOGGER_MAGIC_NUMBER_CUSTOM_VERSION_CONTAINER 0xBFBFBBFB
+#define VISUAL_LOGGER_MAGIC_NUMBER_LATEST VISUAL_LOGGER_MAGIC_NUMBER_CUSTOM_VERSION_CONTAINER
 
 //----------------------------------------------------------------------//
 // FVisualLogShapeElement 
@@ -76,7 +79,7 @@ FVisualLogEntry::FVisualLogEntry(const AActor* InActor, TArray<TWeakObjectPtr<UO
 		Reset();
 	}
 
-	TimeStamp = InActor->GetWorld()->TimeSeconds;
+	TimeStamp = InActor->GetWorld()->TimeSeconds; // Should we be bypassing FVisualLogger::GetTimeStampForObject() ?
 	Location = InActor->GetActorLocation();
 	bIsLocationValid = true;
 
@@ -102,7 +105,7 @@ FVisualLogEntry::FVisualLogEntry(const AActor* InActor, TArray<TWeakObjectPtr<UO
 	}
 }
 
-FVisualLogEntry::FVisualLogEntry(float InTimeStamp, FVector InLocation, const UObject* Object, TArray<TWeakObjectPtr<UObject> >* Children)
+FVisualLogEntry::FVisualLogEntry(double InTimeStamp, FVector InLocation, const UObject* Object, TArray<TWeakObjectPtr<UObject> >* Children)
 {
 	TimeStamp = InTimeStamp;
 	Location = InLocation;
@@ -130,7 +133,7 @@ FVisualLogEntry::FVisualLogEntry(float InTimeStamp, FVector InLocation, const UO
 	}
 }
 
-void FVisualLogEntry::InitializeEntry(const float InTimeStamp)
+void FVisualLogEntry::InitializeEntry(const double InTimeStamp)
 {
 	Reset();
 	TimeStamp = InTimeStamp;
@@ -139,7 +142,7 @@ void FVisualLogEntry::InitializeEntry(const float InTimeStamp)
 
 void FVisualLogEntry::Reset()
 {
-	TimeStamp = -1;
+	TimeStamp = -1.;
 	Location = FVector::ZeroVector;
 	bIsLocationValid = false;
 	Events.Reset();
@@ -245,7 +248,7 @@ void FVisualLogEntry::AddCircle(const FVector& Center, const FVector& UpAxis, co
 	Element.Points.Reserve(3);
 	Element.Points.Add(Center);
 	Element.Points.Add(UpAxis);
-	Element.Points.Add(FVector(Radius, 0.0f, 0.0f));
+	Element.Points.Add(FVector(Radius, 0., 0.));
 	Element.Verbosity = Verbosity;
 	ElementsToDraw.Add(Element);
 }
@@ -548,10 +551,21 @@ FArchive& operator<<(FArchive& Ar, FVisualLogStatusCategory& Status)
 
 FArchive& operator<<(FArchive& Ar, FVisualLogEntry& LogEntry)
 {
-	Ar << LogEntry.TimeStamp;
+	const int32 VLogsOldVer = Ar.CustomVer(EVisualLoggerVersion::GUID);
+	const int32 VLogsStreamObjectVer = Ar.CustomVer(FUE5MainStreamObjectVersion::GUID);
 
-	const int32 VLogsVer = Ar.CustomVer(EVisualLoggerVersion::GUID);
-	if (VLogsVer >= EVisualLoggerVersion::LargeWorldCoordinatesAndLocationValidityFlag)
+	if (VLogsStreamObjectVer >= FUE5MainStreamObjectVersion::VisualLoggerTimeStampAsDouble)
+	{
+		Ar << LogEntry.TimeStamp;
+	}
+	else
+	{
+		float TimeStampFlt = static_cast<float>(LogEntry.TimeStamp);
+		Ar << TimeStampFlt;
+		LogEntry.TimeStamp = TimeStampFlt;
+	}
+
+	if (VLogsOldVer >= EVisualLoggerVersion::LargeWorldCoordinatesAndLocationValidityFlag)
 	{
 		Ar << LogEntry.Location;
 
@@ -572,7 +586,7 @@ FArchive& operator<<(FArchive& Ar, FVisualLogEntry& LogEntry)
 	Ar << LogEntry.ElementsToDraw;
 	Ar << LogEntry.DataBlocks;
 	
-	if (VLogsVer > EVisualLoggerVersion::Initial)
+	if (VLogsOldVer > EVisualLoggerVersion::Initial)
 	{
 		Ar << LogEntry.HistogramSamples;
 	}
@@ -597,7 +611,7 @@ FString FVisualLoggerHelpers::GenerateTemporaryFilename(const FString& FileExt)
 	return FString::Printf(TEXT("VTEMP_%s.%s"), *FDateTime::Now().ToString(), *FileExt);
 }
 
-FString FVisualLoggerHelpers::GenerateFilename(const FString& TempFileName, const FString& Prefix, float StartRecordingTime, float EndTimeStamp)
+FString FVisualLoggerHelpers::GenerateFilename(const FString& TempFileName, const FString& Prefix, double StartRecordingTime, double EndTimeStamp)
 {
 	const FString FullFilename = FString::Printf(TEXT("%s_%s"), *Prefix, *TempFileName);
 	const FString TimeFrameString = FString::Printf(TEXT("%d-%d_"), FMath::TruncToInt(StartRecordingTime), FMath::TruncToInt(EndTimeStamp));
@@ -624,30 +638,42 @@ FArchive& FVisualLoggerHelpers::Serialize(FArchive& Ar, FName& Name)
 FArchive& FVisualLoggerHelpers::Serialize(FArchive& Ar, TArray<FVisualLogDevice::FVisualLogEntryItem>& RecordedLogs)
 {
 	Ar.UsingCustomVersion(EVisualLoggerVersion::GUID);
+	Ar.UsingCustomVersion(FUE5MainStreamObjectVersion::GUID);
 
 	if (Ar.IsLoading())
 	{
 		TArray<FVisualLogDevice::FVisualLogEntryItem> CurrentFrame;
 		while (Ar.AtEnd() == false)
 		{
-			int32 FrameTag = VISUAL_LOGGER_MAGIC_NUMBER;
+			int32 FrameTag = VISUAL_LOGGER_MAGIC_NUMBER_LATEST;
 			Ar << FrameTag;
-			if (FrameTag != DEPRECATED_VISUAL_LOGGER_MAGIC_NUMBER && FrameTag != VISUAL_LOGGER_MAGIC_NUMBER)
+			if (FrameTag != DEPRECATED_VISUAL_LOGGER_MAGIC_NUMBER && FrameTag != VISUAL_LOGGER_MAGIC_NUMBER_OLD_CUSTOM_VERSION && FrameTag != VISUAL_LOGGER_MAGIC_NUMBER_CUSTOM_VERSION_CONTAINER)
 			{
 				break;
 			}
 
-			if (FrameTag == DEPRECATED_VISUAL_LOGGER_MAGIC_NUMBER)
+			if (FrameTag == VISUAL_LOGGER_MAGIC_NUMBER_CUSTOM_VERSION_CONTAINER)
 			{
-				Ar.SetCustomVersion(EVisualLoggerVersion::GUID, EVisualLoggerVersion::Initial, TEXT("VisualLogger"));
+				FCustomVersionContainer CustomVersions;
+				CustomVersions.Serialize(Ar);
+				Ar.SetCustomVersions(CustomVersions);
 			}
 			else
 			{
-				int32 ArchiveVer = -1;
-				Ar << ArchiveVer;
-				check(ArchiveVer >= EVisualLoggerVersion::Initial);
+				Ar.SetCustomVersion(FUE5MainStreamObjectVersion::GUID, FUE5MainStreamObjectVersion::BeforeCustomVersionWasAdded, TEXT("VisualLogger"));
 
-				Ar.SetCustomVersion(EVisualLoggerVersion::GUID, ArchiveVer, TEXT("VisualLogger"));
+				if (FrameTag == VISUAL_LOGGER_MAGIC_NUMBER_OLD_CUSTOM_VERSION)
+				{
+					int32 ArchiveOldVer = -1;
+					Ar << ArchiveOldVer;
+					check(ArchiveOldVer >= 0);
+
+					Ar.SetCustomVersion(EVisualLoggerVersion::GUID, ArchiveOldVer, TEXT("VisualLogger"));
+				}
+				else // DEPRECATED_VISUAL_LOGGER_MAGIC_NUMBER
+				{
+					Ar.SetCustomVersion(EVisualLoggerVersion::GUID, EVisualLoggerVersion::Initial, TEXT("VisualLogger"));
+				}
 			}
 
 			Ar << CurrentFrame;
@@ -657,16 +683,15 @@ FArchive& FVisualLoggerHelpers::Serialize(FArchive& Ar, TArray<FVisualLogDevice:
 	}
 	else
 	{
-		int32 FrameTag = VISUAL_LOGGER_MAGIC_NUMBER;
+		int32 FrameTag = VISUAL_LOGGER_MAGIC_NUMBER_LATEST;
 		Ar << FrameTag;
 
-		int32 ArchiveVer = Ar.CustomVer(EVisualLoggerVersion::GUID);
-		Ar << ArchiveVer;
+		FCustomVersionContainer CustomVersions = Ar.GetCustomVersions();
+		CustomVersions.Serialize(Ar);
+
 		Ar << RecordedLogs;
 	}
 
-	int32 CustomVer = Ar.CustomVer(EVisualLoggerVersion::GUID);
-
 	return Ar;
 }
 
diff --git a/Engine/Source/Runtime/Engine/Public/VisualLogger/VisualLogger.h b/Engine/Source/Runtime/Engine/Public/VisualLogger/VisualLogger.h
--- a/Engine/Source/Runtime/Engine/Public/VisualLogger/VisualLogger.h
+++ b/Engine/Source/Runtime/Engine/Public/VisualLogger/VisualLogger.h
@@ -488,7 +488,7 @@ public:
 	void ClearCategoryAllowList() { CategoryAllowList.Reset(); }
 
 	/** Generates and returns Id unique for given timestamp - used to connect different logs between (ex. text log with geometry shape) */
-	int32 GetUniqueId(float Timestamp);
+	int32 GetUniqueId(double Timestamp);
 
 	/** Starts visual log collecting and recording */
 	void SetIsRecording(const bool bInIsRecording);
@@ -522,7 +522,7 @@ public:
 	/** Returns  current entry for given TimeStamp or creates another one  but first it serialize previous 
 	 *	entry as completed to vislog devices. Use VisualLogger::DontCreate to get current entry without serialization
 	 *	@note this function can return null */
-	FVisualLogEntry* GetEntryToWrite(const UObject* Object, float TimeStamp, ECreateIfNeeded ShouldCreate = ECreateIfNeeded::Create);
+	FVisualLogEntry* GetEntryToWrite(const UObject* Object, double TimeStamp, ECreateIfNeeded ShouldCreate = ECreateIfNeeded::Create);
 	/** Retrieves last used entry for given UObject
 	 *	@note this function can return null */
 	FVisualLogEntry* GetLastEntryForObject(const UObject* Object);
@@ -547,10 +547,10 @@ public:
 	static bool CheckVisualLogInputInternal(const UObject* Object, const FName& CategoryName, ELogVerbosity::Type Verbosity, UWorld **World, FVisualLogEntry **CurrentEntry);
 	
 	/** Returns time stamp for object */
-	float GetTimeStampForObject(const UObject* Object) const;
+	double GetTimeStampForObject(const UObject* Object) const;
 
 	/** Sets function to call to get a timestamp instead of the default implementation (i.e. world time) */
-	void SetGetTimeStampFunc(TFunction<float(const UObject*)> Function);
+	void SetGetTimeStampFunc(TFunction<double(const UObject*)> Function);
 
 	typedef TMap<FObjectKey, TArray<TWeakObjectPtr<const UObject> > > FOwnerToChildrenRedirectionMap;
 	static FOwnerToChildrenRedirectionMap& GetRedirectionMap(const UObject* InObject);
@@ -589,7 +589,7 @@ private:
 	FVisualLoggerObjectEntryMap& GetThreadCurrentEntryMap();
 
 	/** Get global entry to write where all logs are combined together, not thread safe */
-	FVisualLogEntry* GetEntryToWriteInternal(const UObject* Object, float TimeStamp, ECreateIfNeeded ShouldCreate);
+	FVisualLogEntry* GetEntryToWriteInternal(const UObject* Object, double TimeStamp, ECreateIfNeeded ShouldCreate);
 
 	/** Redirect internal implementation, not thread safe */
 	UObject* RedirectInternal(const UObject* FromObject, const UObject* ToObject);
@@ -633,7 +633,7 @@ protected:
 	// Visual Logger extensions map
 	TMap<FName, FVisualLogExtensionInterface*> AllExtensions;
 	// last generated unique id for given times tamp
-	TMap<float, int32> LastUniqueIds;
+	TMap<double, int32> LastUniqueIds;
 	// Current entry with all data
 	FVisualLoggerObjectEntryMap CurrentEntryPerObject;
 	// Threads current entry maps
@@ -667,12 +667,12 @@ protected:
 	/** Indicates there are entries in the redirection map that are invalid */
 	mutable bool bContainsInvalidRedirects : 1;
 	// start recording time
-	float StartRecordingToFileTime;
+	double StartRecordingToFileTime;
 	/** Delegate to set project specific file name for vlogs */
 	FVisualLogFilenameGetterDelegate LogFileNameGetter;
 
 	/** function to call when getting the time stamp */
-	TFunction<float(const UObject*)> GetTimeStampFunc;
+	TFunction<double(const UObject*)> GetTimeStampFunc;
 
 	// if set we are recording and collecting all vlog data
 	static int32 bIsRecording;
diff --git a/Engine/Source/Runtime/Engine/Public/VisualLogger/VisualLoggerTypes.h b/Engine/Source/Runtime/Engine/Public/VisualLogger/VisualLoggerTypes.h
--- a/Engine/Source/Runtime/Engine/Public/VisualLogger/VisualLoggerTypes.h
+++ b/Engine/Source/Runtime/Engine/Public/VisualLogger/VisualLoggerTypes.h
@@ -166,7 +166,7 @@ struct ENGINE_API FVisualLogDataBlock
 struct ENGINE_API FVisualLogEntry
 {
 #if ENABLE_VISUAL_LOG
-	float TimeStamp;
+	double TimeStamp;
 	FVector Location;
 	uint8 bPassedClassAllowList : 1;
 	uint8 bPassedObjectAllowList : 1;	
@@ -184,7 +184,7 @@ struct ENGINE_API FVisualLogEntry
 	FVisualLogEntry() { Reset(); }
 	FVisualLogEntry(const FVisualLogEntry& Entry);
 	FVisualLogEntry(const AActor* InActor, TArray<TWeakObjectPtr<UObject> >* Children);
-	FVisualLogEntry(float InTimeStamp, FVector InLocation, const UObject* Object, TArray<TWeakObjectPtr<UObject> >* Children);
+	FVisualLogEntry(double InTimeStamp, FVector InLocation, const UObject* Object, TArray<TWeakObjectPtr<UObject> >* Children);
 
 	bool ShouldLog(const ECreateIfNeeded ShouldCreate) const
 	{
@@ -193,13 +193,13 @@ struct ENGINE_API FVisualLogEntry
 		return bIsAllowedToLog && ShouldCreate == ECreateIfNeeded::Create;
 	}
 
-	bool ShouldFlush(float InTimeStamp) const
+	bool ShouldFlush(double InTimeStamp) const
 	{
 		//Same LogOwner can be used for logs at different time in the frame so need to flush entry right away
 		return bIsInitialized && InTimeStamp > TimeStamp;
 	}
 
-	void InitializeEntry( const float InTimeStamp );
+	void InitializeEntry( const double InTimeStamp );
 	void Reset();
 	void SetPassedObjectAllowList(const bool bPassed);
 	void UpdateAllowedToLog();
@@ -270,8 +270,12 @@ public:
 	virtual ~FVisualLogDevice() { }
 	virtual void Serialize(const UObject* LogOwner, FName OwnerName, FName InOwnerClassName, const FVisualLogEntry& LogEntry) = 0;
 	virtual void Cleanup(bool bReleaseMemory = false) { /* Empty */ }
-	virtual void StartRecordingToFile(float TImeStamp) { /* Empty */ }
-	virtual void StopRecordingToFile(float TImeStamp) { /* Empty */ }
+	virtual void StartRecordingToFile(double TimeStamp) { /* Empty */ }
+	virtual void StopRecordingToFile(double TimeStamp) { /* Empty */ }
+	UE_DEPRECATED(5.2, "Use the version which takes a double.")
+	virtual void StartRecordingToFile(float TimeStamp) final { StartRecordingToFile(static_cast<double>(TimeStamp)); }
+	UE_DEPRECATED(5.2, "Use the version which takes a double.")
+	virtual void StopRecordingToFile(float TimeStamp) final { StopRecordingToFile(static_cast<double>(TimeStamp)); }
 	virtual void DiscardRecordingToFile() { /* Empty */ }
 	virtual void SetFileName(const FString& InFileName) { /* Empty */ }
 	virtual void GetRecordedLogs(TArray<FVisualLogDevice::FVisualLogEntryItem>& OutLogs)  const { /* Empty */ }
@@ -299,7 +303,7 @@ struct ENGINE_API FVisualLoggerCategoryVerbosityPair
 struct ENGINE_API FVisualLoggerHelpers
 {
 	static FString GenerateTemporaryFilename(const FString& FileExt);
-	static FString GenerateFilename(const FString& TempFileName, const FString& Prefix, float StartRecordingTime, float EndTimeStamp);
+	static FString GenerateFilename(const FString& TempFileName, const FString& Prefix, double StartRecordingTime, double EndTimeStamp);
 	static FArchive& Serialize(FArchive& Ar, FName& Name);
 	static FArchive& Serialize(FArchive& Ar, TArray<FVisualLogDevice::FVisualLogEntryItem>& RecordedLogs);
 	static void GetCategories(const FVisualLogEntry& RecordedLogs, TArray<FVisualLoggerCategoryVerbosityPair>& OutCategories);
@@ -332,7 +336,9 @@ public:
 
 	virtual void OnItemsSelectionChanged(IVisualLoggerEditorInterface* EdInterface) {};
 	virtual void OnLogLineSelectionChanged(IVisualLoggerEditorInterface* EdInterface, TSharedPtr<struct FLogEntryItem> SelectedItem, int64 UserData) {};
-	virtual void OnScrubPositionChanged(IVisualLoggerEditorInterface* EdInterface, float NewScrubPosition, bool bScrubbing) {}
+	UE_DEPRECATED(5.2, "Use the version which takes a double.")
+	virtual void OnScrubPositionChanged(IVisualLoggerEditorInterface* EdInterface, float NewScrubPosition, bool bScrubbing) final {}
+	virtual void OnScrubPositionChanged(IVisualLoggerEditorInterface* EdInterface, double NewScrubPosition, bool bScrubbing) {}
 };
 
 ENGINE_API  FArchive& operator<<(FArchive& Ar, FVisualLogDevice::FVisualLogEntryItem& FrameCacheItem);
