[Visual Logger] expanded subcategories are preserved while scrubbing the timeline

diff --git a/Engine/Source/Developer/LogVisualizer/Private/SVisualLoggerStatusView.cpp b/Engine/Source/Developer/LogVisualizer/Private/SVisualLoggerStatusView.cpp
--- a/Engine/Source/Developer/LogVisualizer/Private/SVisualLoggerStatusView.cpp
+++ b/Engine/Source/Developer/LogVisualizer/Private/SVisualLoggerStatusView.cpp
@@ -143,16 +143,20 @@ void SVisualLoggerStatusView::OnItemSelectionChanged(const FVisualLoggerDBRow& C
 	}
 
 	{
-		for (int32 StatusItemIndex = 0; StatusItemIndex < StatusItems.Num(); StatusItemIndex++)
+		TUniqueFunction<void(TConstArrayView<TSharedPtr<FLogStatusItem>>)> ExpandCategoriesRecursively;
+		ExpandCategoriesRecursively = [&](TConstArrayView< TSharedPtr<FLogStatusItem> > LogStatusItems)
 		{
-			for (const FString& Category : ExpandedCategories)
-			{	if (StatusItems[StatusItemIndex]->ItemText == Category)
+			for (const TSharedPtr<FLogStatusItem>& StatusItem : LogStatusItems)
+			{
+				if (ExpandedCategories.Contains(StatusItem->ItemText))
 				{
-					StatusItemsView->SetItemExpansion(StatusItems[StatusItemIndex], true);
-					break;
+					StatusItemsView->SetItemExpansion(StatusItem, true);
+					ExpandCategoriesRecursively(StatusItem->Children);
 				}
 			}
-		}
+		};
+ 
+		ExpandCategoriesRecursively(StatusItems);
 	}
 }
 
