remaining fixes for renaming TaskGraph -> ProfileVisualizer module

diff --git a/Engine/Source/Developer/LogVisualizer/Private/LogVisualizerStyle.cpp b/Engine/Source/Developer/LogVisualizer/Private/LogVisualizerStyle.cpp
--- a/Engine/Source/Developer/LogVisualizer/Private/LogVisualizerStyle.cpp
+++ b/Engine/Source/Developer/LogVisualizer/Private/LogVisualizerStyle.cpp
@@ -67,8 +67,8 @@ TSharedRef< FSlateStyleSet > FLogVisualizerStyle::Create()
 		Style.Set("LogVisualizerApp.TabIcon", new IMAGE_BRUSH_SVG("Starship/Common/VisualLogger", Icon16x16));
 		
 		Style.Set("LogVisualizer.LogBar.Background", new BOX_BRUSH("Common/ProgressBar_Background", FMargin(5.f / 12.f)));
-		Style.Set("LogVisualizer.LogBar.Selected", new CORE_BOX_BRUSH("Common/TaskGraph_Selected", FMargin(5.f / 12.f)));
-		Style.Set("LogVisualizer.LogBar.EntryDefault", new CORE_BOX_BRUSH("Common/TaskGraph_Mono", FMargin(5.f / 12.f)));
+		Style.Set("LogVisualizer.LogBar.Selected", new CORE_BOX_BRUSH("Common/ProfileVisualizer_Selected", FMargin(5.f / 12.f)));
+		Style.Set("LogVisualizer.LogBar.EntryDefault", new CORE_BOX_BRUSH("Common/ProfileVisualizer_Mono", FMargin(5.f / 12.f)));
 		Style.Set("LogVisualizer.LogBar.TimeMark", new BOX_BRUSH("Icons/LV_BarMark", FMargin(5.f / 12.f)));
 
 		//Style.Set("ToolPanel.GroupBorder", FAppStyle::Get().GetBrush("ToolPanel.GroupBorder"));
diff --git a/Engine/Source/Developer/Profiler/Private/ProfilerStyle.cpp b/Engine/Source/Developer/Profiler/Private/ProfilerStyle.cpp
--- a/Engine/Source/Developer/Profiler/Private/ProfilerStyle.cpp
+++ b/Engine/Source/Developer/Profiler/Private/ProfilerStyle.cpp
@@ -199,7 +199,7 @@ FProfilerStyle::FProfilerStyle()
 
 		Set("Profiler.Misc.Disconnect", new IMAGE_BRUSH("Icons/Profiler/profiler_Disconnect_32x", Icon32x32));
 
-		Set("TaskGraph.Mono", new BOX_BRUSH("Common/TaskGraph_Mono", FMargin(5.f / 12.f)));
+		Set("ProfileVisualizer.Mono", new BOX_BRUSH("Common/ProfileVisualizer_Mono", FMargin(5.f / 12.f)));
 	}
 
 
