Initial AI pass on the Vector2f / FDeprecateSlateVector2D slate changes. We still need to do another pass once the rest of slate and the editor engine code has been converted.

diff --git a/Engine/Plugins/AI/EnvironmentQueryEditor/Source/EnvironmentQueryEditor/Private/SEnvQueryLoadGraph.cpp b/Engine/Plugins/AI/EnvironmentQueryEditor/Source/EnvironmentQueryEditor/Private/SEnvQueryLoadGraph.cpp
--- a/Engine/Plugins/AI/EnvironmentQueryEditor/Source/EnvironmentQueryEditor/Private/SEnvQueryLoadGraph.cpp
+++ b/Engine/Plugins/AI/EnvironmentQueryEditor/Source/EnvironmentQueryEditor/Private/SEnvQueryLoadGraph.cpp
@@ -26,7 +26,7 @@ int32 SEnvQueryLoadGraph::OnPaint(const FPaintArgs& Args, const FGeometry& Allot
 	FSlateDrawElement::MakeBox(
 		OutDrawElements,
 		LayerId,
-		AllottedGeometry.ToPaintGeometry(FVector2D(AllottedGeometry.GetLocalSize().X, AllottedGeometry.GetLocalSize().Y), FSlateLayoutTransform()),
+		AllottedGeometry.ToPaintGeometry(FVector2f(AllottedGeometry.GetLocalSize().X, AllottedGeometry.GetLocalSize().Y), FSlateLayoutTransform()),
 		TimelineAreaBrush,
 		DrawEffects,
 		TimelineAreaBrush->GetTint(InWidgetStyle) * InWidgetStyle.GetColorAndOpacityTint()
diff --git a/Engine/Plugins/AI/EnvironmentQueryEditor/Source/EnvironmentQueryEditor/Private/SEnvQueryProfiler.cpp b/Engine/Plugins/AI/EnvironmentQueryEditor/Source/EnvironmentQueryEditor/Private/SEnvQueryProfiler.cpp
--- a/Engine/Plugins/AI/EnvironmentQueryEditor/Source/EnvironmentQueryEditor/Private/SEnvQueryProfiler.cpp
+++ b/Engine/Plugins/AI/EnvironmentQueryEditor/Source/EnvironmentQueryEditor/Private/SEnvQueryProfiler.cpp
@@ -63,7 +63,7 @@ void SEnvQueryProfiler::Construct(const FArguments& InArgs)
 	VBox->AddSlot().VAlign(VAlign_Top)
 		[
 			SAssignNew(ListView, SListView< FEnvQueryProfilerStatDataPtr >)
-			.ItemHeight(24)
+			.ItemHeight(24.f)
 			.ListItemsSource(&StatData)
 			.OnGenerateRow(this, &SEnvQueryProfiler::OnGenerateRowForList)
 			.OnMouseButtonDoubleClick(this, &SEnvQueryProfiler::OnItemDoubleClicked)
diff --git a/Engine/Plugins/AI/EnvironmentQueryEditor/Source/EnvironmentQueryEditor/Private/STestFunctionWidget.cpp b/Engine/Plugins/AI/EnvironmentQueryEditor/Source/EnvironmentQueryEditor/Private/STestFunctionWidget.cpp
--- a/Engine/Plugins/AI/EnvironmentQueryEditor/Source/EnvironmentQueryEditor/Private/STestFunctionWidget.cpp
+++ b/Engine/Plugins/AI/EnvironmentQueryEditor/Source/EnvironmentQueryEditor/Private/STestFunctionWidget.cpp
@@ -27,7 +27,7 @@ int32 STestFunctionWidget::OnPaint(const FPaintArgs& Args, const FGeometry& Allo
 	FSlateDrawElement::MakeBox(
 		OutDrawElements,
 		LayerId,
-		AllottedGeometry.ToPaintGeometry(FVector2D(AllottedGeometry.Size.X, AllottedGeometry.Size.Y), FSlateLayoutTransform()),
+		AllottedGeometry.ToPaintGeometry(AllottedGeometry.GetLocalSize(), FSlateLayoutTransform()),
 		TimelineAreaBrush,
 		DrawEffects,
 		TimelineAreaBrush->GetTint(InWidgetStyle) * InWidgetStyle.GetColorAndOpacityTint()
@@ -49,7 +49,7 @@ int32 STestFunctionWidget::OnPaint(const FPaintArgs& Args, const FGeometry& Allo
 		FSlateDrawElement::MakeBox(
 			OutDrawElements,
 			LayerId,
-			AllottedGeometry.ToPaintGeometry(FVector2D(FMath::TruncToInt(PreviewData.FilterLow * AllottedGeometry.Size.X), AllottedGeometry.Size.Y), FSlateLayoutTransform(FVector2D(0, 0))),
+			AllottedGeometry.ToPaintGeometry(FVector2f(FMath::TruncToFloat(PreviewData.FilterLow * AllottedGeometry.Size.X), AllottedGeometry.Size.Y), FSlateLayoutTransform(FVector2f(0.f, 0.f))),
 			WhiteBrush,
 			DrawEffects,
 			WhiteBrush->GetTint(InWidgetStyle) * FLinearColor(1.0f, 0.0f, 0.0f, 0.4f)
@@ -63,7 +63,7 @@ int32 STestFunctionWidget::OnPaint(const FPaintArgs& Args, const FGeometry& Allo
 		FSlateDrawElement::MakeBox(
 			OutDrawElements,
 			LayerId,
-			AllottedGeometry.ToPaintGeometry(FVector2D(AllottedGeometry.Size.X, AllottedGeometry.Size.Y), FSlateLayoutTransform(FVector2D(FMath::TruncToInt(PreviewData.FilterHigh * AllottedGeometry.Size.X), 0))),
+			AllottedGeometry.ToPaintGeometry(AllottedGeometry.GetLocalSize(), FSlateLayoutTransform(FVector2f(FMath::TruncToFloat(PreviewData.FilterHigh * AllottedGeometry.Size.X), 0.f))),
 			WhiteBrush,
 			DrawEffects,
 			WhiteBrush->GetTint(InWidgetStyle) * FLinearColor(1.0f, 0.0f, 0.0f, 0.5f)
@@ -171,7 +171,7 @@ int32 STestFunctionWidget::OnPaint(const FPaintArgs& Args, const FGeometry& Allo
 		const float XPos = Idx * DeltaX;
 		const float YPos = (AllottedGeometry.Size.Y - 1) - (PreviewData.Samples[Idx] * AllottedGeometry.Size.Y);
 
-		LinePoints.Add(FVector2D(FMath::TruncToInt(XPos), FMath::TruncToInt(YPos)));
+		LinePoints.Add(FVector2D(FMath::TruncToFloat(XPos), FMath::TruncToFloat(YPos)));
 	}
 
 	FSlateDrawElement::MakeLines(
diff --git a/Engine/Plugins/Runtime/MassEntity/Source/MassEntityDebugger/Private/MassDebuggerStyle.cpp b/Engine/Plugins/Runtime/MassEntity/Source/MassEntityDebugger/Private/MassDebuggerStyle.cpp
--- a/Engine/Plugins/Runtime/MassEntity/Source/MassEntityDebugger/Private/MassDebuggerStyle.cpp
+++ b/Engine/Plugins/Runtime/MassEntity/Source/MassEntityDebugger/Private/MassDebuggerStyle.cpp
@@ -29,7 +29,7 @@ public:
 	FMassDebuggerStyleSet(const FName& InStyleSetName)
 		: FSlateStyleSet(InStyleSetName)
 	{
-		const FVector2D Icon16x16(16.0f, 16.0f);
+		const FVector2f Icon16x16(16.0f, 16.0f);
 
 		SetContentRoot(FPaths::EngineContentDir() / TEXT("Editor/Slate"));
 		SetCoreContentRoot(FPaths::EngineContentDir() / TEXT("Slate"));
diff --git a/Engine/Plugins/Runtime/MassEntity/Source/MassEntityDebugger/Private/SMassProcessorsView.cpp b/Engine/Plugins/Runtime/MassEntity/Source/MassEntityDebugger/Private/SMassProcessorsView.cpp
--- a/Engine/Plugins/Runtime/MassEntity/Source/MassEntityDebugger/Private/SMassProcessorsView.cpp
+++ b/Engine/Plugins/Runtime/MassEntity/Source/MassEntityDebugger/Private/SMassProcessorsView.cpp
@@ -98,7 +98,7 @@ void SMassProcessorsView::Construct(const FArguments& InArgs, TSharedRef<FMassDe
 			.MinSize(250.0f)
 			[
 				SAssignNew(ProcessorsListWidget, SListView<TSharedPtr<FMassDebuggerProcessorData> >)
-					.ItemHeight(20)
+					.ItemHeight(20.f)
 					.ListItemsSource(&DebuggerModel->CachedProcessors)
 					.SelectionMode(ESelectionMode::Multi)
 					.OnSelectionChanged(this, &SMassProcessorsView::ProcessorListSelectionChanged)
diff --git a/Engine/Plugins/Runtime/MassEntity/Source/MassEntityEditor/Private/MassEditorStyle.cpp b/Engine/Plugins/Runtime/MassEntity/Source/MassEntityEditor/Private/MassEditorStyle.cpp
--- a/Engine/Plugins/Runtime/MassEntity/Source/MassEntityEditor/Private/MassEditorStyle.cpp
+++ b/Engine/Plugins/Runtime/MassEntity/Source/MassEntityEditor/Private/MassEditorStyle.cpp
@@ -29,7 +29,7 @@ void FMassEntityEditorStyle::Initialize()
 		return;
 	}
 
-	const FVector2D Icon8x8(8.0f, 8.0f);
+	const FVector2f Icon8x8(8.0f, 8.0f);
 
 	StyleSet = MakeShared<FSlateStyleSet>(GetStyleSetName());
 
diff --git a/Engine/Plugins/Runtime/StateTree/Source/StateTreeEditorModule/Private/StateTreeEditorStyle.cpp b/Engine/Plugins/Runtime/StateTree/Source/StateTreeEditorModule/Private/StateTreeEditorStyle.cpp
--- a/Engine/Plugins/Runtime/StateTree/Source/StateTreeEditorModule/Private/StateTreeEditorStyle.cpp
+++ b/Engine/Plugins/Runtime/StateTree/Source/StateTreeEditorModule/Private/StateTreeEditorStyle.cpp
@@ -32,7 +32,7 @@ void FStateTreeEditorStyle::Initialize()
 		return;
 	}
 
-	const FVector2D Icon8x8(8.0f, 8.0f);
+	const FVector2f Icon8x8(8.0f, 8.0f);
 
 	StyleSet = MakeShared<FSlateStyleSet>(GetStyleSetName());
 
diff --git a/Engine/Source/Developer/LogVisualizer/Private/LogVisualizerStyle.cpp b/Engine/Source/Developer/LogVisualizer/Private/LogVisualizerStyle.cpp
--- a/Engine/Source/Developer/LogVisualizer/Private/LogVisualizerStyle.cpp
+++ b/Engine/Source/Developer/LogVisualizer/Private/LogVisualizerStyle.cpp
@@ -38,11 +38,11 @@ FName FLogVisualizerStyle::GetStyleSetName()
 
 TSharedRef< FSlateStyleSet > FLogVisualizerStyle::Create()
 {
-	const FVector2D Icon8x8(8.0f, 8.0f);
-	const FVector2D Icon16x16(16.0f, 16.0f);
-	const FVector2D Icon20x20(20.0f, 20.0f);
-	const FVector2D Icon40x40(40.0f, 40.0f);
-	const FVector2D Icon48x48(48.0f, 48.0f);
+	const FVector2f Icon8x8(8.0f, 8.0f);
+	const FVector2f Icon16x16(16.0f, 16.0f);
+	const FVector2f Icon20x20(20.0f, 20.0f);
+	const FVector2f Icon40x40(40.0f, 40.0f);
+	const FVector2f Icon48x48(48.0f, 48.0f);
 
 	TSharedRef<FSlateStyleSet> StyleRef = MakeShareable(new FSlateStyleSet(FLogVisualizerStyle::GetStyleSetName()));
 	StyleRef->SetContentRoot(FPaths::EngineContentDir() / TEXT("Editor/Slate"));
@@ -53,7 +53,7 @@ TSharedRef< FSlateStyleSet > FLogVisualizerStyle::Create()
 	const FTextBlockStyle NormalText = FTextBlockStyle()
 		.SetFont(DEFAULT_FONT("Regular", 9))
 		.SetColorAndOpacity(FSlateColor::UseForeground())
-		.SetShadowOffset(FVector2D::ZeroVector)
+		.SetShadowOffset(FVector2f::ZeroVector)
 		.SetShadowColorAndOpacity(FLinearColor::Black)
 		.SetHighlightColor(FLinearColor(0.02f, 0.3f, 0.0f))
 		.SetHighlightShape(BOX_BRUSH("Common/TextBlockHighlightShape", FMargin(3.f / 8.f)));
@@ -104,9 +104,9 @@ TSharedRef< FSlateStyleSet > FLogVisualizerStyle::Create()
 
 		Style.Set("TableView.DarkRow", FTableRowStyle(NormalTableRowStyle)
 			.SetEvenRowBackgroundBrush(IMAGE_BRUSH("Common/Selection", Icon8x8, FLinearColor(1.0f, 1.0f, 1.0f, 0.1f)))
-			.SetEvenRowBackgroundHoveredBrush(IMAGE_BRUSH("PropertyView/DetailCategoryMiddle_Hovered", FVector2D(16, 16)))
+			.SetEvenRowBackgroundHoveredBrush(IMAGE_BRUSH("PropertyView/DetailCategoryMiddle_Hovered", FVector2f(16.f, 16.f)))
 			.SetOddRowBackgroundBrush(IMAGE_BRUSH("Common/Selection", Icon8x8, FLinearColor(1.0f, 1.0f, 1.0f, 0.1f)))
-			.SetOddRowBackgroundHoveredBrush(IMAGE_BRUSH("PropertyView/DetailCategoryMiddle_Hovered", FVector2D(16, 16)))
+			.SetOddRowBackgroundHoveredBrush(IMAGE_BRUSH("PropertyView/DetailCategoryMiddle_Hovered", FVector2f(16.f, 16.f)))
 			.SetSelectorFocusedBrush(BORDER_BRUSH("Common/Selector", FMargin(4.f / 16.f), SelectorColor))
 			.SetActiveBrush(IMAGE_BRUSH("Common/Selection", Icon8x8, SelectionColor))
 			.SetActiveHoveredBrush(IMAGE_BRUSH("Common/Selection", Icon8x8, SelectionColor))
@@ -119,7 +119,7 @@ TSharedRef< FSlateStyleSet > FLogVisualizerStyle::Create()
 	Style.Set("GenericFilters.TextStyle", FTextBlockStyle(NormalText)
 		.SetFont(DEFAULT_FONT("Bold", 9))
 		.SetColorAndOpacity(FLinearColor(1.0f, 1.0f, 1.0f, 0.9f))
-		.SetShadowOffset(FVector2D(1, 1))
+		.SetShadowOffset(FVector2f(1.f, 1.f))
 		.SetShadowColorAndOpacity(FLinearColor(0, 0, 0, 0.9f)));
 
 	// Toolbar
@@ -159,12 +159,12 @@ TSharedRef< FSlateStyleSet > FLogVisualizerStyle::Create()
 		
 		//Style.Set("ContentBrowser.FilterButton", FAppStyle::Get().GetWidgetStyle<FCheckBoxStyle>("ContentBrowser.FilterButton"));
 		const FCheckBoxStyle ContentBrowserFilterButtonCheckBoxStyle = FCheckBoxStyle()
-			.SetUncheckedImage(IMAGE_BRUSH("ContentBrowser/FilterUnchecked", FVector2D(10.0f, 14.0f)))
-			.SetUncheckedHoveredImage(IMAGE_BRUSH("ContentBrowser/FilterUnchecked", FVector2D(10.0f, 14.0f), FLinearColor(0.5f, 0.5f, 0.5f, 1.0f)))
-			.SetUncheckedPressedImage(IMAGE_BRUSH("ContentBrowser/FilterUnchecked", FVector2D(10.0f, 14.0f), FLinearColor(0.5f, 0.5f, 0.5f, 1.0f)))
-			.SetCheckedImage(IMAGE_BRUSH("ContentBrowser/FilterChecked", FVector2D(10.0f, 14.0f)))
-			.SetCheckedHoveredImage(IMAGE_BRUSH("ContentBrowser/FilterChecked", FVector2D(10.0f, 14.0f), FLinearColor(0.5f, 0.5f, 0.5f, 1.0f)))
-			.SetCheckedPressedImage(IMAGE_BRUSH("ContentBrowser/FilterChecked", FVector2D(10.0f, 14.0f), FLinearColor(0.5f, 0.5f, 0.5f, 1.0f)));
+			.SetUncheckedImage(IMAGE_BRUSH("ContentBrowser/FilterUnchecked", FVector2f(10.0f, 14.0f)))
+			.SetUncheckedHoveredImage(IMAGE_BRUSH("ContentBrowser/FilterUnchecked", FVector2f(10.0f, 14.0f), FLinearColor(0.5f, 0.5f, 0.5f, 1.0f)))
+			.SetUncheckedPressedImage(IMAGE_BRUSH("ContentBrowser/FilterUnchecked", FVector2f(10.0f, 14.0f), FLinearColor(0.5f, 0.5f, 0.5f, 1.0f)))
+			.SetCheckedImage(IMAGE_BRUSH("ContentBrowser/FilterChecked", FVector2f(10.0f, 14.0f)))
+			.SetCheckedHoveredImage(IMAGE_BRUSH("ContentBrowser/FilterChecked", FVector2f(10.0f, 14.0f), FLinearColor(0.5f, 0.5f, 0.5f, 1.0f)))
+			.SetCheckedPressedImage(IMAGE_BRUSH("ContentBrowser/FilterChecked", FVector2f(10.0f, 14.0f), FLinearColor(0.5f, 0.5f, 0.5f, 1.0f)));
 		/* ... and add the new style */
 		Style.Set("ContentBrowser.FilterButton", ContentBrowserFilterButtonCheckBoxStyle);
 	}
@@ -181,7 +181,7 @@ TSharedRef< FSlateStyleSet > FLogVisualizerStyle::Create()
 
 	// Default text styles
 	{
-		Style.Set("RichText.Background", new BOX_BRUSH("Common/FlatColorSquare", FVector2D(1.0f, 1.0f), FMargin(0), FLinearColor(FColor(0xffeff3f3))));
+		Style.Set("RichText.Background", new BOX_BRUSH("Common/FlatColorSquare", FVector2f(1.0f, 1.0f), FMargin(0), FLinearColor(FColor(0xffeff3f3))));
 
 		Style.Set("RichText.RoundedBackground", new BOX_BRUSH("Common/RoundedSelection_16x", 4.0f / 16.0f, FLinearColor(FColor(0xffeff3f3))));
 
@@ -231,8 +231,8 @@ TSharedRef< FSlateStyleSet > FLogVisualizerStyle::Create()
 
 		//Tagline
 		{
-			Style.Set("RichText.Tagline.Background", new BOX_BRUSH("Common/FlatColorSquare", FVector2D(1.0f, 1.0f), FMargin(1), FLinearColor(FColor(0xffdbe4e4))));
-			Style.Set("RichText.Tagline.DarkBackground", new BOX_BRUSH("Common/FlatColorSquare", FVector2D(1.0f, 1.0f), FMargin(1), FLinearColor(0.55423, 0.60548, 0.60548)));
+			Style.Set("RichText.Tagline.Background", new BOX_BRUSH("Common/FlatColorSquare", FVector2f(1.0f, 1.0f), FMargin(1), FLinearColor(FColor(0xffdbe4e4))));
+			Style.Set("RichText.Tagline.DarkBackground", new BOX_BRUSH("Common/FlatColorSquare", FVector2f(1.0f, 1.0f), FMargin(1), FLinearColor(0.55423, 0.60548, 0.60548)));
 			Style.Set("RichText.Tagline.Text", FTextBlockStyle(NormalText)
 				.SetFont(DEFAULT_FONT("Bold", 24))
 				.SetColorAndOpacity(FLinearColor(FColor(0xff2c3e50)))
diff --git a/Engine/Source/Developer/LogVisualizer/Private/SVisualLoggerTimelinesContainer.cpp b/Engine/Source/Developer/LogVisualizer/Private/SVisualLoggerTimelinesContainer.cpp
--- a/Engine/Source/Developer/LogVisualizer/Private/SVisualLoggerTimelinesContainer.cpp
+++ b/Engine/Source/Developer/LogVisualizer/Private/SVisualLoggerTimelinesContainer.cpp
@@ -355,7 +355,7 @@ void SVisualLoggerTimelinesContainer::OnFiltersSearchChanged(const FText& Filter
 void SVisualLoggerTimelinesContainer::GenerateReport()
 {
 	TSharedRef<SWindow> NewWindow = SNew(SWindow)
-		.ClientSize(FVector2D(720, 768))
+		.ClientSize(FVector2f(720.f, 768.f))
 		.Title(NSLOCTEXT("LogVisualizerReport", "WindowTitle", "Log Visualizer Report"))
 		[
 			SNew(SVisualLoggerReport, CachedSelectedTimelines, VisualLoggerView)
diff --git a/Engine/Source/Developer/LogVisualizer/Private/VisualLoggerTimeSliderController.cpp b/Engine/Source/Developer/LogVisualizer/Private/VisualLoggerTimeSliderController.cpp
--- a/Engine/Source/Developer/LogVisualizer/Private/VisualLoggerTimeSliderController.cpp
+++ b/Engine/Source/Developer/LogVisualizer/Private/VisualLoggerTimeSliderController.cpp
@@ -24,32 +24,31 @@ namespace ScrubConstants
 /** Utility struct for converting between scrub range space and local/absolute screen space */
 struct FVisualLoggerTimeSliderController::FScrubRangeToScreen
 {
-	FVector2D WidgetSize;
+	FVector2f WidgetSize;
 
 	TRange<double> ViewInput;
 	double ViewInputRange;
 	double PixelsPerInput;
 
-	FScrubRangeToScreen(TRange<double> InViewInput, const FVector2D& InWidgetSize )
+	FScrubRangeToScreen(TRange<double> InViewInput, const FVector2f& InWidgetSize )
 	{
 		WidgetSize = InWidgetSize;
 
 		ViewInput = InViewInput;
 		ViewInputRange = ViewInput.Size<double>();
-		PixelsPerInput = ViewInputRange > 0 ? ( WidgetSize.X / ViewInputRange ) : 0;
+		PixelsPerInput = ViewInputRange > 0. ? ( WidgetSize.X / ViewInputRange ) : 0.;
 	}
 
 	/** Local Widget Space -> Curve Input domain. */
-	double LocalXToInput(double ScreenX) const
+	double LocalXToInput(float ScreenX) const
 	{
-		double LocalX = ScreenX;
-		return (LocalX/PixelsPerInput) + ViewInput.GetLowerBoundValue();
+		return (ScreenX/PixelsPerInput) + ViewInput.GetLowerBoundValue();
 	}
 
 	/** Curve Input domain -> local Widget Space */
-	double InputToLocalX(double Input) const
+	float InputToLocalX(double Input) const
 	{
-		return (Input - ViewInput.GetLowerBoundValue()) * PixelsPerInput;
+		return static_cast<float>((Input - ViewInput.GetLowerBoundValue()) * PixelsPerInput);
 	}
 };
 
@@ -117,9 +116,9 @@ struct FVisualLoggerTimeSliderController::FDrawTickArgs
 	/** Color of each tick */
 	FLinearColor TickColor;
 	/** Offset in Y where to start the tick */
-	double TickOffset;
+	float TickOffset;
 	/** Height in of major ticks */
-	double MajorTickHeight;
+	float MajorTickHeight;
 	/** Start layer for elements */
 	int32 StartLayer;
 	/** Draw effects to apply */
@@ -145,22 +144,22 @@ void FVisualLoggerTimeSliderController::DrawTicks( FSlateWindowElementList& OutD
 	
 	FSlateFontInfo SmallLayoutFont = FCoreStyle::GetDefaultFontStyle("Regular", 8);
 
-	TArray<FVector2D> LinePoints;
+	TArray<FVector2f> LinePoints;
 	LinePoints.AddUninitialized(2);
 
 	double Seconds = 0;
 	while( (Seconds = OffsetNum*Spacing) < RangeToScreen.ViewInput.GetUpperBoundValue() )
 	{
 		// X position local to start of the widget area
-		double XPos = RangeToScreen.InputToLocalX( Seconds );
+		const float XPos = RangeToScreen.InputToLocalX( Seconds );
 		uint32 AbsOffsetNum = FMath::Abs(OffsetNum);
 
 		if ( AbsOffsetNum % Divider == 0 )
 		{
-			FVector2D Offset( XPos, InArgs.TickOffset );
-			FVector2D TickSize( 1.0f, InArgs.MajorTickHeight );
+			FVector2f Offset( XPos, InArgs.TickOffset );
+			FVector2f TickSize( 1.0f, InArgs.MajorTickHeight );
 
-			LinePoints[0] = FVector2D(1.0f,1.0f);
+			LinePoints[0] = FVector2f(1.0f,1.0f);
 			LinePoints[1] = TickSize;
 
 			// lines should not need anti-aliasing
@@ -183,8 +182,8 @@ void FVisualLoggerTimeSliderController::DrawTicks( FSlateWindowElementList& OutD
 
 				// Space the text between the tick mark but slightly above
 				const TSharedRef< FSlateFontMeasure > FontMeasureService = FSlateApplication::Get().GetRenderer()->GetFontMeasureService();
-				FVector2D TextSize = FontMeasureService->Measure(FrameString, SmallLayoutFont);
-				FVector2D TextOffset( XPos-(TextSize.X*0.5), InArgs.bMirrorLabels ? TextSize.Y :  FMath::Abs( InArgs.AllottedGeometry.GetLocalSize().Y - (InArgs.MajorTickHeight+TextSize.Y) ) );
+				const FVector2f TextSize = FontMeasureService->Measure(FrameString, SmallLayoutFont);
+				const FVector2f TextOffset( XPos-(TextSize.X*0.5f), InArgs.bMirrorLabels ? TextSize.Y :  FMath::Abs( InArgs.AllottedGeometry.GetLocalSize().Y - (InArgs.MajorTickHeight+TextSize.Y) ) );
 
 				FSlateDrawElement::MakeText(
 					OutDrawElements,
@@ -202,10 +201,10 @@ void FVisualLoggerTimeSliderController::DrawTicks( FSlateWindowElementList& OutD
 			// Compute the size of each tick mark.  If we are half way between to visible values display a slightly larger tick mark
 			const float MinorTickHeight = AbsOffsetNum % HalfDivider == 0 ? 7.0f : 4.0f;
 
-			FVector2D Offset(XPos, InArgs.bMirrorLabels ? 0.0f : FMath::Abs( InArgs.AllottedGeometry.GetLocalSize().Y - MinorTickHeight ) );
-			FVector2D TickSize(1, MinorTickHeight);
+			FVector2f Offset(XPos, InArgs.bMirrorLabels ? 0.0f : FMath::Abs( InArgs.AllottedGeometry.GetLocalSize().Y - MinorTickHeight ) );
+			FVector2f TickSize(1, MinorTickHeight);
 
-			LinePoints[0] = FVector2D(1.0f,1.0f);
+			LinePoints[0] = FVector2f(1.0f,1.0f);
 			LinePoints[1] = TickSize;
 
 			const bool bAntiAlias = false;
@@ -241,7 +240,7 @@ int32 FVisualLoggerTimeSliderController::OnPaintTimeSlider( bool bMirrorLabels,
 	{
 		FScrubRangeToScreen RangeToScreen( LocalViewRange, AllottedGeometry.GetLocalSize() );
 	
-		const float MajorTickHeight = 9.;
+		const float MajorTickHeight = 9.f;
 	
 		FDrawTickArgs Args;
 		Args.AllottedGeometry = AllottedGeometry;
@@ -251,7 +250,7 @@ int32 FVisualLoggerTimeSliderController::OnPaintTimeSlider( bool bMirrorLabels,
 		Args.ClippingRect = MyCullingRect;
 		Args.DrawEffects = DrawEffects;
 		Args.StartLayer = LayerId;
-		Args.TickOffset = bMirrorLabels ? 0. : FMath::Abs( AllottedGeometry.GetLocalSize().Y - MajorTickHeight );
+		Args.TickOffset = bMirrorLabels ? 0.f : FMath::Abs( AllottedGeometry.GetLocalSize().Y - MajorTickHeight );
 		Args.MajorTickHeight = MajorTickHeight;
 
 		DrawTicks( OutDrawElements, RangeToScreen, Args );
@@ -260,13 +259,13 @@ int32 FVisualLoggerTimeSliderController::OnPaintTimeSlider( bool bMirrorLabels,
 		float HalfSize = FMath::TruncToFloat(HandleSize/2.0f);
 
 		// Draw the scrub handle
-		const double XPos = RangeToScreen.InputToLocalX( TimeSliderArgs.ScrubPosition.Get() );
+		const float XPos = RangeToScreen.InputToLocalX( TimeSliderArgs.ScrubPosition.Get() );
 
 		// Draw cursor size
 		const float CursorHalfSize = TimeSliderArgs.CursorSize.Get() * 0.5f;
 		const int32 CursorLayer = LayerId + 2;
-		const double CursorHalfLength = AllottedGeometry.GetLocalSize().X * CursorHalfSize;
-		FPaintGeometry CursorGeometry = AllottedGeometry.ToPaintGeometry(FVector2f(2 * CursorHalfLength, AllottedGeometry.GetLocalSize().Y), FSlateLayoutTransform(FVector2f(XPos - CursorHalfLength, 0)));
+		const float CursorHalfLength = AllottedGeometry.GetLocalSize().X * CursorHalfSize;
+		FPaintGeometry CursorGeometry = AllottedGeometry.ToPaintGeometry(FVector2f(2.f * CursorHalfLength, AllottedGeometry.GetLocalSize().Y), FSlateLayoutTransform(FVector2f(XPos - CursorHalfLength, 0.f)));
 
 		FLinearColor CursorColor = InWidgetStyle.GetColorAndOpacityTint();
 		CursorColor.A = CursorColor.A*0.08f;
@@ -315,9 +314,9 @@ FReply FVisualLoggerTimeSliderController::OnMouseButtonDown( SWidget& WidgetOwne
 	if ( bHandleLeftMouseButton )
 	{
 		// Always capture mouse if we left or right click on the widget
-		FScrubRangeToScreen RangeToScreen(TimeSliderArgs.ViewRange.Get(), MyGeometry.GetLocalSize());
-		FVector2D CursorPos = MyGeometry.AbsoluteToLocal(MouseEvent.GetLastScreenSpacePosition());
-		double NewValue = RangeToScreen.LocalXToInput(CursorPos.X);
+		const FScrubRangeToScreen RangeToScreen(TimeSliderArgs.ViewRange.Get(), MyGeometry.GetLocalSize());
+		const FVector2f CursorPos = MyGeometry.AbsoluteToLocal(MouseEvent.GetLastScreenSpacePosition());
+		const double NewValue = RangeToScreen.LocalXToInput(CursorPos.X);
 
 		CommitScrubPosition(NewValue, /*bIsScrubbing=*/false);
 		return FReply::Handled().CaptureMouse( WidgetOwner.AsShared() ).PreventThrottling();
@@ -354,9 +353,9 @@ FReply FVisualLoggerTimeSliderController::OnMouseButtonUp( SWidget& WidgetOwner,
 		}
 		else
 		{
-			FScrubRangeToScreen RangeToScreen( TimeSliderArgs.ViewRange.Get(), MyGeometry.GetLocalSize() );
-			FVector2D CursorPos = MyGeometry.AbsoluteToLocal(MouseEvent.GetLastScreenSpacePosition());
-			double NewValue = RangeToScreen.LocalXToInput(CursorPos.X);
+			const FScrubRangeToScreen RangeToScreen( TimeSliderArgs.ViewRange.Get(), MyGeometry.GetLocalSize() );
+			const FVector2f CursorPos = MyGeometry.AbsoluteToLocal(MouseEvent.GetLastScreenSpacePosition());
+			const double NewValue = RangeToScreen.LocalXToInput(CursorPos.X);
 
 			CommitScrubPosition( NewValue, /*bIsScrubbing=*/false );
 		}
@@ -371,12 +370,12 @@ FReply FVisualLoggerTimeSliderController::OnMouseButtonUp( SWidget& WidgetOwner,
 
 double FVisualLoggerTimeSliderController::GetTimeAtCursorPosition(const FGeometry& MyGeometry, const FPointerEvent& MouseEvent) const
 {
-	FScrubRangeToScreen RangeToScreen(TimeSliderArgs.ViewRange.Get(), MyGeometry.GetLocalSize());
-	FVector2D CursorPos = MyGeometry.AbsoluteToLocal(MouseEvent.GetLastScreenSpacePosition());
+	const FScrubRangeToScreen RangeToScreen(TimeSliderArgs.ViewRange.Get(), MyGeometry.GetLocalSize());
+	const FVector2f CursorPos = MyGeometry.AbsoluteToLocal(MouseEvent.GetLastScreenSpacePosition());
 	double NewValue = RangeToScreen.LocalXToInput(CursorPos.X);
 
-	double LocalClampMin = TimeSliderArgs.ClampRange.Get().GetLowerBoundValue();
-	double LocalClampMax = TimeSliderArgs.ClampRange.Get().GetUpperBoundValue();
+	const double LocalClampMin = TimeSliderArgs.ClampRange.Get().GetLowerBoundValue();
+	const double LocalClampMax = TimeSliderArgs.ClampRange.Get().GetUpperBoundValue();
 
 	if (NewValue < LocalClampMin)
 	{
@@ -416,7 +415,7 @@ FReply FVisualLoggerTimeSliderController::OnMouseMove( SWidget& WidgetOwner, con
 				double LocalViewRangeMax = LocalViewRange.GetUpperBoundValue();
 
 				FScrubRangeToScreen ScaleInfo( LocalViewRange, MyGeometry.GetLocalSize() );
-				FVector2D ScreenDelta = MouseEvent.GetCursorDelta();
+				FVector2f ScreenDelta = MouseEvent.GetCursorDelta();
 				FVector2D InputDelta;
 				InputDelta.X = ScreenDelta.X/ScaleInfo.PixelsPerInput;
 
@@ -638,7 +637,7 @@ int32 FVisualLoggerTimeSliderController::OnPaintSectionView( const FGeometry& Al
 		// Draw major ticks under sections
 		Args.StartLayer = LayerId-1;
 		// Draw the tick the entire height of the section area
-		Args.TickOffset = 0.;
+		Args.TickOffset = 0.f;
 		Args.MajorTickHeight = AllottedGeometry.GetLocalSize().Y;
 
 		DrawTicks( OutDrawElements, RangeToScreen, Args );
diff --git a/Engine/Source/Developer/LogVisualizer/Private/VisualLoggerTimeSliderController.h b/Engine/Source/Developer/LogVisualizer/Private/VisualLoggerTimeSliderController.h
--- a/Engine/Source/Developer/LogVisualizer/Private/VisualLoggerTimeSliderController.h
+++ b/Engine/Source/Developer/LogVisualizer/Private/VisualLoggerTimeSliderController.h
@@ -132,5 +132,5 @@ private:
 	bool bPanning;
 	/***/
 	TSharedPtr<SScrollBar> Scrollbar;
-	FVector2D SoftwareCursorPosition;
+	FVector2f SoftwareCursorPosition;
 };
diff --git a/Engine/Source/Editor/AIGraph/Private/AIGraphConnectionDrawingPolicy.cpp b/Engine/Source/Editor/AIGraph/Private/AIGraphConnectionDrawingPolicy.cpp
--- a/Engine/Source/Editor/AIGraph/Private/AIGraphConnectionDrawingPolicy.cpp
+++ b/Engine/Source/Editor/AIGraph/Private/AIGraphConnectionDrawingPolicy.cpp
@@ -83,21 +83,23 @@ void FAIGraphConnectionDrawingPolicy::Internal_DrawLineWithArrow(const FVector2D
 	//@TODO: Should this be scaled by zoom factor?
 	const float LineSeparationAmount = 4.5f;
 
-	const FVector2D DeltaPos = EndAnchorPoint - StartAnchorPoint;
-	const FVector2D UnitDelta = DeltaPos.GetSafeNormal();
-	const FVector2D Normal = FVector2D(DeltaPos.Y, -DeltaPos.X).GetSafeNormal();
+	const FVector2f StartAnchorPointTemp(UE::Slate::CastToVector2f(StartAnchorPoint));
+	const FVector2f EndAnchorPointTemp(static_cast<float>(EndAnchorPoint.X), static_cast<float>(EndAnchorPoint.Y));
+	const FVector2f DeltaPos = EndAnchorPointTemp - StartAnchorPointTemp;
+	const FVector2f UnitDelta = DeltaPos.GetSafeNormal();
+	const FVector2f Normal = FVector2f(DeltaPos.Y, -DeltaPos.X).GetSafeNormal();
 
 	// Come up with the final start/end points
-	const FVector2D DirectionBias = Normal * LineSeparationAmount;
-	const FVector2D LengthBias = ArrowRadius.X * UnitDelta;
-	const FVector2D StartPoint = StartAnchorPoint + DirectionBias + LengthBias;
-	const FVector2D EndPoint = EndAnchorPoint + DirectionBias - LengthBias;
+	const FVector2f DirectionBias = Normal * LineSeparationAmount;
+	const FVector2f LengthBias = ArrowRadius.X * UnitDelta;
+	const FDeprecateSlateVector2D StartPoint = StartAnchorPointTemp + DirectionBias + LengthBias;
+	const FDeprecateSlateVector2D EndPoint = EndAnchorPointTemp + DirectionBias - LengthBias;
 
 	// Draw a line/spline
-	DrawConnection(WireLayerID, StartPoint, EndPoint, Params);
+	DrawConnection(WireLayerID, FVector2D(StartPoint), FVector2D(EndPoint), Params);
 
 	// Draw the arrow
-	const FVector2D ArrowDrawPos = EndPoint - ArrowRadius;
+	const FVector2f ArrowDrawPos = EndPoint - ArrowRadius;
 	const float AngleInRadians = static_cast<float>(FMath::Atan2(DeltaPos.Y, DeltaPos.X));
 
 	FSlateDrawElement::MakeRotatedBox(
@@ -107,7 +109,7 @@ void FAIGraphConnectionDrawingPolicy::Internal_DrawLineWithArrow(const FVector2D
 		ArrowImage,
 		ESlateDrawEffect::None,
 		AngleInRadians,
-		TOptional<FVector2D>(),
+		TOptional<FVector2f>(),
 		FSlateDrawElement::RelativeToElement,
 		Params.WireColor
 		);
diff --git a/Engine/Source/Editor/BehaviorTreeEditor/Classes/BehaviorTreeEditorTypes.h b/Engine/Source/Editor/BehaviorTreeEditor/Classes/BehaviorTreeEditorTypes.h
--- a/Engine/Source/Editor/BehaviorTreeEditor/Classes/BehaviorTreeEditorTypes.h
+++ b/Engine/Source/Editor/BehaviorTreeEditor/Classes/BehaviorTreeEditorTypes.h
@@ -11,6 +11,7 @@
 #include "UObject/Object.h"
 #include "UObject/ObjectMacros.h"
 #include "UObject/UObjectGlobals.h"
+#include "Types/SlateVector2.h"
 
 #include "BehaviorTreeEditorTypes.generated.h"
 
@@ -52,13 +53,13 @@ namespace ESubNode
 
 struct FNodeBounds
 {
-	FVector2D Position;
-	FVector2D Size;
+	FDeprecateSlateVector2D Position;
+	FDeprecateSlateVector2D Size;
 
-	FNodeBounds(FVector2D InPos, FVector2D InSize)
+	FNodeBounds(UE::Slate::FDeprecateVector2DParameter InPos, UE::Slate::FDeprecateVector2DParameter InSize)
+		: Position(InPos)
+		, Size(InSize)
 	{
-		Position = InPos;
-		Size = InSize;
 	}
 };
 
diff --git a/Engine/Source/Editor/BehaviorTreeEditor/Private/AssetTypeActions_BehaviorTree.cpp b/Engine/Source/Editor/BehaviorTreeEditor/Private/AssetTypeActions_BehaviorTree.cpp
--- a/Engine/Source/Editor/BehaviorTreeEditor/Private/AssetTypeActions_BehaviorTree.cpp
+++ b/Engine/Source/Editor/BehaviorTreeEditor/Private/AssetTypeActions_BehaviorTree.cpp
@@ -100,7 +100,7 @@ void FAssetTypeActions_BehaviorTree::PerformAssetDiff(UObject* OldAsset, UObject
 
 	const TSharedPtr<SWindow> Window = SNew(SWindow)
 		.Title(WindowTitle)
-		.ClientSize(FVector2D(1000,800));
+		.ClientSize(FVector2f(1000.f,800.f));
 
 	Window->SetContent(SNew(SBehaviorTreeDiff)
 		.BehaviorTreeOld(OldBehaviorTree)
diff --git a/Engine/Source/Editor/BehaviorTreeEditor/Private/BehaviorTreeGraph.cpp b/Engine/Source/Editor/BehaviorTreeEditor/Private/BehaviorTreeGraph.cpp
--- a/Engine/Source/Editor/BehaviorTreeEditor/Private/BehaviorTreeGraph.cpp
+++ b/Engine/Source/Editor/BehaviorTreeEditor/Private/BehaviorTreeGraph.cpp
@@ -1210,7 +1210,7 @@ namespace BTAutoArrangeHelpers
 {
 	struct FNodeBoundsInfo
 	{
-		FVector2D SubGraphBBox;
+		FDeprecateSlateVector2D SubGraphBBox;
 		TArray<FNodeBoundsInfo> Children;
 	};
 
@@ -1229,7 +1229,7 @@ namespace BTAutoArrangeHelpers
 				if (GraphNode && BBoxTree.Children.Num() > 0)
 				{
 					AutoArrangeNodes(GraphNode, BBoxTree.Children[BBoxIndex], PosX, PosY + GraphNode->DEPRECATED_NodeWidget.Pin()->GetDesiredSize().Y * 2.5f);
-					GraphNode->DEPRECATED_NodeWidget.Pin()->MoveTo(FVector2D(BBoxTree.Children[BBoxIndex].SubGraphBBox.X / 2.f - GraphNode->DEPRECATED_NodeWidget.Pin()->GetDesiredSize().X / 2.f + PosX, PosY), NodeFilter);
+					GraphNode->DEPRECATED_NodeWidget.Pin()->MoveTo(FDeprecateSlateVector2D(BBoxTree.Children[BBoxIndex].SubGraphBBox.X / 2.f - GraphNode->DEPRECATED_NodeWidget.Pin()->GetDesiredSize().X / 2.f + PosX, PosY), NodeFilter);
 					PosX += BBoxTree.Children[BBoxIndex].SubGraphBBox.X + 20.f;
 					BBoxIndex++;
 				}
diff --git a/Engine/Source/Editor/BehaviorTreeEditor/Private/SBehaviorTreeDiff.cpp b/Engine/Source/Editor/BehaviorTreeEditor/Private/SBehaviorTreeDiff.cpp
--- a/Engine/Source/Editor/BehaviorTreeEditor/Private/SBehaviorTreeDiff.cpp
+++ b/Engine/Source/Editor/BehaviorTreeEditor/Private/SBehaviorTreeDiff.cpp
@@ -291,7 +291,7 @@ TSharedRef<SWidget> SBehaviorTreeDiff::GenerateDiffListWidget()
 			.FillHeight(1.f)
 			[
 				SAssignNew(DiffList, SListViewType)
-				.ItemHeight(24)
+				.ItemHeight(24.f)
 				.ListItemsSource(&DiffListSource)
 				.OnGenerateRow(this, &SBehaviorTreeDiff::OnGenerateRow)
 				.SelectionMode(ESelectionMode::Single)
diff --git a/Engine/Source/Editor/BehaviorTreeEditor/Private/SGraphNode_BehaviorTree.cpp b/Engine/Source/Editor/BehaviorTreeEditor/Private/SGraphNode_BehaviorTree.cpp
--- a/Engine/Source/Editor/BehaviorTreeEditor/Private/SGraphNode_BehaviorTree.cpp
+++ b/Engine/Source/Editor/BehaviorTreeEditor/Private/SGraphNode_BehaviorTree.cpp
@@ -632,7 +632,7 @@ void SGraphNode_BehaviorTree::PostChange(const UUserDefinedEnum* Changed, FEnumE
 void SGraphNode_BehaviorTree::Tick( const FGeometry& AllottedGeometry, const double InCurrentTime, const float InDeltaTime )
 {
 	SGraphNode::Tick( AllottedGeometry, InCurrentTime, InDeltaTime );
-	CachedPosition = FVector2D(AllottedGeometry.AbsolutePosition) / AllottedGeometry.Scale;
+	CachedPosition = AllottedGeometry.AbsolutePosition / AllottedGeometry.Scale;
 
 	UBehaviorTreeGraphNode* MyNode = Cast<UBehaviorTreeGraphNode>(GraphNode);
 	if (MyNode && MyNode->DebuggerUpdateCounter != DebuggerStateCounter)
@@ -710,7 +710,7 @@ void SGraphNode_BehaviorTree::Tick( const FGeometry& AllottedGeometry, const dou
 			// when it wasn't any of them, add node itself to triggers (e.g. parallel's main task)
 			if (DecoratorWidgets.Num() == 0)
 			{
-				TriggerOffsets.Add(FNodeBounds(FVector2D(0,0),GetDesiredSize()));
+				TriggerOffsets.Add(FNodeBounds(FVector2f(0.f,0.f),GetDesiredSize()));
 			}
 		}
 	}
@@ -1061,7 +1061,7 @@ void SGraphNode_BehaviorTree::GetOverlayBrushes(bool bSelected, const FVector2D
 				for (int32 i = 0; i < TriggerOffsets.Num(); i++)
 				{
 					IPOverlayInfo.OverlayOffset.X = -IPOverlayInfo.Brush->ImageSize.X;
-					IPOverlayInfo.OverlayOffset.Y = TriggerOffsets[i].Position.Y + TriggerOffsets[i].Size.Y / 2 - IPOverlayInfo.Brush->ImageSize.Y / 2;
+					IPOverlayInfo.OverlayOffset.Y = TriggerOffsets[i].Position.Y + TriggerOffsets[i].Size.Y / 2.f - IPOverlayInfo.Brush->ImageSize.Y / 2.f;
 
 					IPOverlayInfo.AnimationEnvelope = FVector2D(10.f, 0.f);
 					Brushes.Add(IPOverlayInfo);
diff --git a/Engine/Source/Editor/BehaviorTreeEditor/Private/SGraphNode_BehaviorTree.h b/Engine/Source/Editor/BehaviorTreeEditor/Private/SGraphNode_BehaviorTree.h
--- a/Engine/Source/Editor/BehaviorTreeEditor/Private/SGraphNode_BehaviorTree.h
+++ b/Engine/Source/Editor/BehaviorTreeEditor/Private/SGraphNode_BehaviorTree.h
@@ -16,6 +16,7 @@
 #include "Styling/SlateColor.h"
 #include "Templates/SharedPointer.h"
 #include "Widgets/DeclarativeSyntaxSupport.h"
+#include "Types/SlateVector2.h"
 
 class SBorder;
 class SGraphNode;
@@ -71,7 +72,7 @@ public:
 	/** shows red marker when search failed*/
 	EVisibility GetDebuggerSearchFailedMarkerVisibility() const;
 
-	FVector2D GetCachedPosition() const { return CachedPosition; }
+	UE::Slate::FDeprecateVector2DResult GetCachedPosition() const { return CachedPosition; }
 
 protected:
 	/** INotifyOnEnumChanged interface */
@@ -95,7 +96,7 @@ protected:
 	TArray<FNodeBounds> TriggerOffsets;
 
 	/** cached draw position */
-	FVector2D CachedPosition;
+	FDeprecateSlateVector2D CachedPosition;
 
 	TArray< TSharedPtr<SGraphNode> > DecoratorWidgets;
 	TArray< TSharedPtr<SGraphNode> > ServicesWidgets;
