[VisualLogger] make sure a newly created row is hidden when all added items are filtered out by categories verbosity

diff --git a/Engine/Source/Developer/LogVisualizer/Private/SVisualLogger.cpp b/Engine/Source/Developer/LogVisualizer/Private/SVisualLogger.cpp
--- a/Engine/Source/Developer/LogVisualizer/Private/SVisualLogger.cpp
+++ b/Engine/Source/Developer/LogVisualizer/Private/SVisualLogger.cpp
@@ -1033,6 +1033,11 @@ void SVisualLogger::OnFiltersSearchChanged(const FText& Filter)
 void SVisualLogger::OnNewItemHandler(const FVisualLoggerDBRow& DBRow, int32 ItemIndex)
 {
 	UpdateVisibilityForEntry(DBRow, ItemIndex);
+
+	// A new item can be hidden by the category filters by UpdateVisibilityForEntry.
+	// In such case, we might need to update row visibility/
+	// Note that this is not called within UpdateVisibilityForEntry since it can be called by async tasks.
+	FVisualLoggerDatabase::Get().SetRowVisibility(DBRow.GetOwnerName(), DBRow.GetNumberOfHiddenItems() != DBRow.GetItems().Num());
 }
 
 void SVisualLogger::UpdateVisibilityForEntry(const FVisualLoggerDBRow& DBRow, int32 ItemIndex)
