[VisualLogger] Fixed regression for entry boxes after FVector2D deprecation pass (cl 23631073)

diff --git a/Engine/Source/Developer/LogVisualizer/Private/SVisualLoggerTimelineBar.cpp b/Engine/Source/Developer/LogVisualizer/Private/SVisualLoggerTimelineBar.cpp
--- a/Engine/Source/Developer/LogVisualizer/Private/SVisualLoggerTimelineBar.cpp
+++ b/Engine/Source/Developer/LogVisualizer/Private/SVisualLoggerTimelineBar.cpp
@@ -242,8 +242,8 @@ int32 SVisualLoggerTimelineBar::OnPaint(const FPaintArgs& Args, const FGeometry&
 				OutDrawElements,
 				RetLayerId,
 				AllottedGeometry.ToPaintGeometry(
-				FVector2f(StartPos, 0.0f),
-				FSlateLayoutTransform(FVector2f(BarWidth, AllottedGeometry.Size.Y))),
+				FVector2f(BarWidth, AllottedGeometry.GetLocalSize().Y),
+				FSlateLayoutTransform(FVector2f(StartPos, 0.0f))),
 				FillImage,
 				DrawEffects,
 				CurrentTimeColor
