[VisualLogger] fixed GotoNextItem and GotoPreviousItem in FLogVisualizer.
#rnx

diff --git a/Engine/Source/Developer/LogVisualizer/Private/LogVisualizer.cpp b/Engine/Source/Developer/LogVisualizer/Private/LogVisualizer.cpp
--- a/Engine/Source/Developer/LogVisualizer/Private/LogVisualizer.cpp
+++ b/Engine/Source/Developer/LogVisualizer/Private/LogVisualizer.cpp
@@ -265,7 +265,6 @@ int32 FLogVisualizer::GetPreviousItem(FName RowName, int32 MoveDistance)
 	return NewItemIndex;
 }
 
-// @todo: This function currently doesn't do anything!
 void FLogVisualizer::GotoNextItem(FName RowName, int32 MoveDistance)
 {
 	FVisualLoggerDBRow& DBRow = FVisualLoggerDatabase::Get().GetRowByName(RowName);
@@ -273,12 +272,10 @@ void FLogVisualizer::GotoNextItem(FName RowName, int32 MoveDistance)
 
 	if (NewItemIndex != DBRow.GetCurrentItemIndex())
 	{
-		auto& Entries = DBRow.GetItems();
-		float NewTimeStamp = Entries[NewItemIndex].Entry.TimeStamp;
+		TimeSliderController->CommitScrubPosition(DBRow.GetItems()[NewItemIndex].Entry.TimeStamp, /*bIsScrubbing*/false);
 	}
 }
 
-// @todo: This function currently doesn't do anything!
 void FLogVisualizer::GotoPreviousItem(FName RowName, int32 MoveDistance)
 {
 	FVisualLoggerDBRow& DBRow = FVisualLoggerDatabase::Get().GetRowByName(RowName);
@@ -286,8 +283,7 @@ void FLogVisualizer::GotoPreviousItem(FName RowName, int32 MoveDistance)
 
 	if (NewItemIndex != DBRow.GetCurrentItemIndex())
 	{
-		auto& Entries = DBRow.GetItems();
-		float NewTimeStamp = Entries[NewItemIndex].Entry.TimeStamp;
+		TimeSliderController->CommitScrubPosition(DBRow.GetItems()[NewItemIndex].Entry.TimeStamp, /*bIsScrubbing*/false);
 	}
 }
 
@@ -309,7 +305,7 @@ void FLogVisualizer::GotoFirstItem(FName RowName)
 	if (NewItemIndex != DBRow.GetCurrentItemIndex())
 	{
 		//DBRow.MoveTo(NewItemIndex);
-		TimeSliderController->CommitScrubPosition(Entries[NewItemIndex].Entry.TimeStamp, false);
+		TimeSliderController->CommitScrubPosition(Entries[NewItemIndex].Entry.TimeStamp, /*bIsScrubbing*/false);
 	}
 }
 
@@ -332,7 +328,7 @@ void FLogVisualizer::GotoLastItem(FName RowName)
 	if (NewItemIndex != DBRow.GetCurrentItemIndex())
 	{
 		//DBRow.MoveTo(NewItemIndex);
-		TimeSliderController->CommitScrubPosition(Entries[NewItemIndex].Entry.TimeStamp, false);
+		TimeSliderController->CommitScrubPosition(Entries[NewItemIndex].Entry.TimeStamp, /*bIsScrubbing*/false);
 	}
 }
 
