Fixed position of debug text for UE_VLOG_OBOX

diff --git a/Engine/Source/Developer/LogVisualizer/Private/VisualLoggerRenderingActorBase.cpp b/Engine/Source/Developer/LogVisualizer/Private/VisualLoggerRenderingActorBase.cpp
--- a/Engine/Source/Developer/LogVisualizer/Private/VisualLoggerRenderingActorBase.cpp
+++ b/Engine/Source/Developer/LogVisualizer/Private/VisualLoggerRenderingActorBase.cpp
@@ -308,15 +308,16 @@ void AVisualLoggerRenderingActorBase::GetDebugShapes(const FVisualLogEntry& InEn
 			const FVector* BoxExtent = ElementToDraw->Points.GetData();
 			const int32 NumPoints = ElementToDraw->Points.Num();
 
+			FTransform Transform(ElementToDraw->TransformationMatrix);
 			for (int32 Index = 0; Index + 1 < NumPoints; Index += 2, BoxExtent += 2)
 			{
 				const FBox Box = FBox(*BoxExtent, *(BoxExtent + 1));
-				DebugShapes.Boxes.Add(FDebugRenderSceneProxy::FDebugBox(Box, Color, FTransform(ElementToDraw->TransformationMatrix)));
+				DebugShapes.Boxes.Add(FDebugRenderSceneProxy::FDebugBox(Box, Color, Transform));
 
 				if (bDrawLabel)
 				{
 					const FString PrintString = NumPoints == 2 ? ElementToDraw->Description : FString::Printf(TEXT("%s_%d"), *ElementToDraw->Description, Index / 2);
-					DebugShapes.Texts.Add(FDebugRenderSceneProxy::FText3d(PrintString, Box.GetCenter(), Color));
+					DebugShapes.Texts.Add(FDebugRenderSceneProxy::FText3d(PrintString, Transform.TransformPosition(Box.GetCenter()), Color));
 				}
 			}
 		}
