Prevent Visual logger rendering actor from being selectable and showing up in the scene outliner

diff --git a/Engine/Source/Developer/LogVisualizer/Private/VisualLoggerDatabase.cpp b/Engine/Source/Developer/LogVisualizer/Private/VisualLoggerDatabase.cpp
--- a/Engine/Source/Developer/LogVisualizer/Private/VisualLoggerDatabase.cpp
+++ b/Engine/Source/Developer/LogVisualizer/Private/VisualLoggerDatabase.cpp
@@ -471,6 +471,11 @@ AActor* FVisualLoggerEditorInterface::GetHelperActor(UWorld* InWorld) const
 	// The helper actor is created on demand and only once per world so we can allow it to spawn during construction script.
 	SpawnInfo.bAllowDuringConstructionScript = true;
 
+#if WITH_EDITOR
+	// Nothing to set on this actor, se we just hide it from the outliner : 
+	SpawnInfo.bHideFromSceneOutliner = true;
+#endif //WITH_EDITOR
+
 	return World->SpawnActor<AVisualLoggerRenderingActor>(SpawnInfo);
 }
 
diff --git a/Engine/Source/Developer/LogVisualizer/Public/VisualLoggerRenderingActorBase.h b/Engine/Source/Developer/LogVisualizer/Public/VisualLoggerRenderingActorBase.h
--- a/Engine/Source/Developer/LogVisualizer/Public/VisualLoggerRenderingActorBase.h
+++ b/Engine/Source/Developer/LogVisualizer/Public/VisualLoggerRenderingActorBase.h
@@ -53,6 +53,10 @@ public:
 	virtual void IterateDebugShapes(TFunction<void(const FTimelineDebugShapes&) > Callback) { };
 	virtual bool MatchCategoryFilters(const FName& CategoryName, ELogVerbosity::Type Verbosity) const { return true; };
 
+#if WITH_EDITOR
+	virtual bool IsSelectable() const override { return false; }
+#endif // WITH_EDITOR
+
 protected:
 #if ENABLE_VISUAL_LOG
 	void GetDebugShapes(const FVisualLogEntry& EntryItem, bool bAddEntryLocationPointer, FTimelineDebugShapes& OutDebugShapes);
