Visual Logger: UX Improvements

diff --git a/Engine/Source/Developer/LogVisualizer/Private/LogVisualizer.cpp b/Engine/Source/Developer/LogVisualizer/Private/LogVisualizer.cpp
--- a/Engine/Source/Developer/LogVisualizer/Private/LogVisualizer.cpp
+++ b/Engine/Source/Developer/LogVisualizer/Private/LogVisualizer.cpp
@@ -21,7 +21,7 @@ FColor FLogVisualizerColorPalette[] = {
 	FColor(0xff8fbc8f), // darkseagreen
 	FColor(0xff48d1cc), // mediumturquoise
 	FColor(0xffffa500), // orange
-	FColor(0xff6a5acd), // slateblue
+	FColor(0xff1e90ff), // dodgerblue
 	FColor(0xffffff00), // yellow
 	FColor(0xffdeb887), // burlywood
 	FColor(0xff00ff00), // lime
@@ -30,7 +30,6 @@ FColor FLogVisualizerColorPalette[] = {
 	FColor(0xff00bfff), // deepskyblue
 	FColor(0xffadff2f), // greenyellow
 	FColor(0xffff00ff), // fuchsia
-	FColor(0xff1e90ff), // dodgerblue
 	FColor(0xfff0e68c), // khaki
 	FColor(0xffdda0dd), // plum
 	FColor(0xff90ee90), // lightgreen
@@ -49,6 +48,7 @@ FColor FLogVisualizerColorPalette[] = {
 	FColor(0xffb03060), // maroon3
 	FColor(0xff8b4513), // saddlebrown
 	FColor(0xff228b22), // forestgreen
+	FColor(0xff6a5acd), // slateblue
 	FColor(0xff808000), // olive
 	FColor(0xff3cb371), // mediumseagreen
 	FColor(0xffb8860b), // darkgoldenrod
diff --git a/Engine/Source/Developer/LogVisualizer/Private/LogVisualizerStyle.cpp b/Engine/Source/Developer/LogVisualizer/Private/LogVisualizerStyle.cpp
--- a/Engine/Source/Developer/LogVisualizer/Private/LogVisualizerStyle.cpp
+++ b/Engine/Source/Developer/LogVisualizer/Private/LogVisualizerStyle.cpp
@@ -76,7 +76,7 @@ TSharedRef< FSlateStyleSet > FLogVisualizerStyle::Create()
 		//Style.Set("NoBorder", FAppStyle::Get().GetBrush("NoBorder"));
 		Style.Set("NoBorder", new FSlateNoResource());
 		Style.Set("ToolBar.Button.Normal", new FSlateNoResource());
-		Style.Set("ToolBar.Button.Hovered", new BOX_BRUSH("Common/RoundedSelection_16x", 4.0f / 16.0f, FLinearColor(0.728f, 0.364f, 0.003f)));
+		Style.Set("ToolBar.Button.Hovered", new BOX_BRUSH("Common/RoundedSelection_16x", 4.0f / 16.0f, FLinearColor(0.1f, 0.1f, 0.1f)));
 
 	}
 	{
diff --git a/Engine/Source/Developer/LogVisualizer/Private/SVisualLogger.cpp b/Engine/Source/Developer/LogVisualizer/Private/SVisualLogger.cpp
--- a/Engine/Source/Developer/LogVisualizer/Private/SVisualLogger.cpp
+++ b/Engine/Source/Developer/LogVisualizer/Private/SVisualLogger.cpp
@@ -343,7 +343,7 @@ void SVisualLogger::Construct(const FArguments& InArgs, const TSharedRef<SDockTa
 				ISettingsModule* SettingsModule = FModuleManager::GetModulePtr<ISettingsModule>("Settings");
 				if (SettingsModule != nullptr)
 				{
-					SettingsModule->ShowViewer("Editor", "General", "VisualLogger");
+					SettingsModule->ShowViewer("Editor", "Advanced", "VisualLogger");
 				}
 			}
 		)),
diff --git a/Engine/Source/Developer/LogVisualizer/Private/SVisualLoggerFilterWidget.cpp b/Engine/Source/Developer/LogVisualizer/Private/SVisualLoggerFilterWidget.cpp
--- a/Engine/Source/Developer/LogVisualizer/Private/SVisualLoggerFilterWidget.cpp
+++ b/Engine/Source/Developer/LogVisualizer/Private/SVisualLoggerFilterWidget.cpp
@@ -51,6 +51,8 @@ void SVisualLoggerFilterWidget::Construct(const FArguments& InArgs)
 					SNew(STextBlock)
 					.ColorAndOpacity(this, &SVisualLoggerFilterWidget::GetFilterNameColorAndOpacity)
 					.Text(this, &SVisualLoggerFilterWidget::GetCaptionString)
+					.ShadowColorAndOpacity(FLinearColor::Black)
+					.ShadowOffset(FVector2D(1.0f, 1.0f))
 				]
 			]
 		];
diff --git a/Engine/Source/Developer/LogVisualizer/Private/SVisualLoggerFilterWidget.h b/Engine/Source/Developer/LogVisualizer/Private/SVisualLoggerFilterWidget.h
--- a/Engine/Source/Developer/LogVisualizer/Private/SVisualLoggerFilterWidget.h
+++ b/Engine/Source/Developer/LogVisualizer/Private/SVisualLoggerFilterWidget.h
@@ -129,8 +129,9 @@ private:
 	/** Handler to determine the color of the checkbox text when it is hovered */
 	FSlateColor GetFilterNameColorAndOpacity() const
 	{
-		const float DimFactor = 0.75f;
-		return IsChecked() == ECheckBoxState::Checked ? (IsHovered() ? ColorCategory * DimFactor : ColorCategory) : (IsHovered() ? FLinearColor::White : FLinearColor::White * DimFactor);
+		const float DimFactor = 0.875f;
+		constexpr FLinearColor UncheckedColor(0.35f, 0.35f, 0.35f);
+		return IsChecked() == ECheckBoxState::Checked ? (IsHovered() ? ColorCategory * DimFactor : ColorCategory) : (IsHovered() ? UncheckedColor : UncheckedColor * DimFactor);
 	}
 
 	FSlateColor GetBorderBackgroundColor() const
diff --git a/Engine/Source/Developer/LogVisualizer/Private/SVisualLoggerFilters.cpp b/Engine/Source/Developer/LogVisualizer/Private/SVisualLoggerFilters.cpp
--- a/Engine/Source/Developer/LogVisualizer/Private/SVisualLoggerFilters.cpp
+++ b/Engine/Source/Developer/LogVisualizer/Private/SVisualLoggerFilters.cpp
@@ -427,12 +427,12 @@ void SVisualLoggerFilters::OnItemsSelectionChanged(const FVisualLoggerDBRow& Cha
 	for (int32 Index = 0; Index < Filters.Num(); ++Index)
 	{
 		SVisualLoggerFilterWidget& Filter = Filters[Index].Get();
-		Filter.SetBorderBackgroundColor(FLinearColor(0.2f, 0.2f, 0.2f, 0.2f));
+		Filter.SetBorderBackgroundColor(FLinearColor(0.2f, 0.2f, 0.2f, 0.1f));
 		for (const FVisualLoggerCategoryVerbosityPair& Category : Categories)
 		{
 			if (Filter.GetFilterName() == Category.CategoryName)
 			{
-				Filter.SetBorderBackgroundColor(FLinearColor(0.3f, 0.3f, 0.3f, 0.8f));
+				Filter.SetBorderBackgroundColor(FLinearColor(0.2f, 0.2f, 0.2f, 0.45f));
 			}
 		}
 	}
