Added Auto-Scroll button to visual logger

diff --git a/Engine/Source/Developer/LogVisualizer/Private/LogVisualizer.cpp b/Engine/Source/Developer/LogVisualizer/Private/LogVisualizer.cpp
--- a/Engine/Source/Developer/LogVisualizer/Private/LogVisualizer.cpp
+++ b/Engine/Source/Developer/LogVisualizer/Private/LogVisualizer.cpp
@@ -145,7 +145,7 @@ UWorld* FLogVisualizer::GetWorld(UObject* OptionalObject)
 void FLogVisualizer::UpdateCameraPosition(FName RowName, int32 ItemIndes)
 {
 	const FVisualLoggerDBRow& DBRow = FVisualLoggerDatabase::Get().GetRowByName(RowName);
-	auto &Entries = DBRow.GetItems();
+	const TArray<FVisualLogDevice::FVisualLogEntryItem>& Entries = DBRow.GetItems();
 	if (DBRow.GetCurrentItemIndex() == INDEX_NONE || Entries.IsValidIndex(DBRow.GetCurrentItemIndex()) == false)
 	{
 		return;
@@ -220,7 +220,7 @@ int32 FLogVisualizer::GetNextItem(FName RowName, int32 MoveDistance)
 	int32 NewItemIndex = DBRow.GetCurrentItemIndex();
 
 	int32 Index = 0;
-	auto &Entries = DBRow.GetItems();
+	const TArray<FVisualLogDevice::FVisualLogEntryItem>& Entries = DBRow.GetItems();
 	while (true)
 	{
 		NewItemIndex++;
@@ -247,7 +247,7 @@ int32 FLogVisualizer::GetPreviousItem(FName RowName, int32 MoveDistance)
 	int32 NewItemIndex = DBRow.GetCurrentItemIndex();
 
 	int32 Index = 0;
-	auto &Entries = DBRow.GetItems();
+	const TArray<FVisualLogDevice::FVisualLogEntryItem>& Entries = DBRow.GetItems();
 	while (true)
 	{
 		NewItemIndex--;
@@ -294,7 +294,7 @@ void FLogVisualizer::GotoFirstItem(FName RowName)
 	FVisualLoggerDBRow& DBRow = FVisualLoggerDatabase::Get().GetRowByName(RowName);
 	int32 NewItemIndex = DBRow.GetCurrentItemIndex();
 
-	auto &Entries = DBRow.GetItems();
+	const TArray<FVisualLogDevice::FVisualLogEntryItem>& Entries = DBRow.GetItems();
 	for (int32 Index = 0; Index <= DBRow.GetCurrentItemIndex(); Index++)
 	{
 		if (DBRow.IsItemVisible(Index))
@@ -316,7 +316,7 @@ void FLogVisualizer::GotoLastItem(FName RowName)
 	FVisualLoggerDBRow& DBRow = FVisualLoggerDatabase::Get().GetRowByName(RowName);
 	int32 NewItemIndex = DBRow.GetCurrentItemIndex();
 
-	auto &Entries = DBRow.GetItems();
+	const TArray<FVisualLogDevice::FVisualLogEntryItem>& Entries = DBRow.GetItems();
 	for (int32 Index = Entries.Num() - 1; Index >= DBRow.GetCurrentItemIndex(); Index--)
 	{
 		if (DBRow.IsItemVisible(Index))
@@ -326,7 +326,6 @@ void FLogVisualizer::GotoLastItem(FName RowName)
 		}
 	}
 
-
 	if (NewItemIndex != DBRow.GetCurrentItemIndex())
 	{
 		//DBRow.MoveTo(NewItemIndex);
@@ -334,6 +333,31 @@ void FLogVisualizer::GotoLastItem(FName RowName)
 	}
 }
 
+void FLogVisualizer::GotoLastItemAnyRow()
+{
+	double LastItemTimeStamp = -1.0;
+	for (FVisualLoggerDatabase::FConstRowIterator Iter(FVisualLoggerDatabase::Get().GetConstRowIterator()); Iter; ++Iter)
+	{
+		const TArray<FVisualLogDevice::FVisualLogEntryItem>& Entries = Iter->GetItems();
+		for (int32 ItemIndex = Entries.Num() - 1; ItemIndex >= 0; ItemIndex--)
+		{
+			if (Iter->IsItemVisible(ItemIndex))
+			{
+				if (Entries[ItemIndex].Entry.TimeStamp > LastItemTimeStamp)
+				{
+					LastItemTimeStamp = Entries[ItemIndex].Entry.TimeStamp;
+					break;
+				}
+			}
+		}
+	}
+
+	if (LastItemTimeStamp > -1.0)
+	{
+		TimeSliderController->CommitScrubPosition(LastItemTimeStamp, /*bIsScrubbing*/false);
+	}
+}
+
 void FLogVisualizer::SeekToTime(float Time)
 {
 	GetTimeSliderController()->CommitScrubPosition(Time,  /*bIsScrubbing=*/true);
diff --git a/Engine/Source/Developer/LogVisualizer/Private/LogVisualizerStyle.cpp b/Engine/Source/Developer/LogVisualizer/Private/LogVisualizerStyle.cpp
--- a/Engine/Source/Developer/LogVisualizer/Private/LogVisualizerStyle.cpp
+++ b/Engine/Source/Developer/LogVisualizer/Private/LogVisualizerStyle.cpp
@@ -135,6 +135,7 @@ TSharedRef< FSlateStyleSet > FLogVisualizerStyle::Create()
 		Style.Set("Toolbar.Remove", new IMAGE_BRUSH("Icons/Edit/icon_Edit_Delete_40x", Icon40x40));
 		Style.Set("Toolbar.Graphs", new IMAGE_BRUSH("Icons/icon_MatEd_Stats_40x", Icon40x40));
 		Style.Set("Toolbar.HideEmpty", new IMAGE_BRUSH("Icons/icon_levels_visible_40x", Icon40x40));
+		Style.Set("Toolbar.AutoScroll", new CORE_IMAGE_BRUSH_SVG("Starship/Insights/AutoScrollRight_20", CoreStyleConstants::Icon16x16));
 
 		Style.Set("Toolbar.Pause.small", new IMAGE_BRUSH("Icons/icon_pause_40x", Icon20x20));
 		Style.Set("Toolbar.Resume.small", new IMAGE_BRUSH("Icons/icon_simulate_40x", Icon20x20));
@@ -147,6 +148,7 @@ TSharedRef< FSlateStyleSet > FLogVisualizerStyle::Create()
 		Style.Set("Toolbar.Remove.small", new IMAGE_BRUSH("Icons/Edit/icon_Edit_Delete_40x", Icon20x20));
 		Style.Set("Toolbar.Graphs.small", new IMAGE_BRUSH("Icons/icon_MatEd_Stats_40x", Icon20x20));
 		Style.Set("Toolbar.HideEmpty.small", new IMAGE_BRUSH("Icons/icon_levels_visible_40x", Icon20x20));
+		Style.Set("Toolbar.AutoScroll,small", new CORE_IMAGE_BRUSH_SVG("Starship/Insights/AutoScrollRight_20", CoreStyleConstants::Icon16x16));
 	}
 
 	// Filters
diff --git a/Engine/Source/Developer/LogVisualizer/Private/SVisualLogger.cpp b/Engine/Source/Developer/LogVisualizer/Private/SVisualLogger.cpp
--- a/Engine/Source/Developer/LogVisualizer/Private/SVisualLogger.cpp
+++ b/Engine/Source/Developer/LogVisualizer/Private/SVisualLogger.cpp
@@ -89,6 +89,8 @@ SVisualLogger::SVisualLogger()
 { 
 	bPausedLogger = false;
 	bGotHistogramData = false;
+	bAutoScrollToLastItem = false;
+	bAutoScrollingToLastItem = false;
 
 	class FVisualLoggerDevice : public FVisualLogDevice
 	{
@@ -162,14 +164,16 @@ void SVisualLogger::Construct(const FArguments& InArgs, const TSharedRef<SDockTa
 {
 	bPausedLogger = false;
 	bGotHistogramData = false;
+	bAutoScrollToLastItem = false;
+	bAutoScrollingToLastItem = false;
 
 	FLogVisualizer::Get().SetCurrentVisualizer(SharedThis(this));
 	//////////////////////////////////////////////////////////////////////////
 	// Visual Logger Events
 	FLogVisualizer::Get().GetEvents().OnFiltersChanged.AddRaw(this, &SVisualLogger::OnFiltersChanged);
-	FLogVisualizer::Get().GetEvents().OnLogLineSelectionChanged = FOnLogLineSelectionChanged::CreateRaw(this, &SVisualLogger::OnLogLineSelectionChanged);
-	FLogVisualizer::Get().GetEvents().OnKeyboardEvent = FOnKeyboardEvent::CreateRaw(this, &SVisualLogger::OnKeyboardRedirection);
-	FLogVisualizer::Get().GetTimeSliderController().Get()->GetTimeSliderArgs().OnScrubPositionChanged = FVisualLoggerTimeSliderArgs::FOnScrubPositionChanged::CreateRaw(this, &SVisualLogger::OnScrubPositionChanged);
+	FLogVisualizer::Get().GetEvents().OnLogLineSelectionChanged = FOnLogLineSelectionChanged::CreateSP(this, &SVisualLogger::OnLogLineSelectionChanged);
+	FLogVisualizer::Get().GetEvents().OnKeyboardEvent = FOnKeyboardEvent::CreateSP(this, &SVisualLogger::OnKeyboardRedirection);
+	FLogVisualizer::Get().GetTimeSliderController().Get()->GetTimeSliderArgs().OnScrubPositionChanged = FVisualLoggerTimeSliderArgs::FOnScrubPositionChanged::CreateSP(this, &SVisualLogger::OnScrubPositionChanged);
 
 	FVisualLoggerDatabase::Get().GetEvents().OnRowSelectionChanged.AddRaw(this, &SVisualLogger::OnObjectSelectionChanged);
 	FVisualLoggerDatabase::Get().GetEvents().OnNewItem.AddRaw(this, &SVisualLogger::OnNewItemHandler);
@@ -219,60 +223,65 @@ void SVisualLogger::Construct(const FArguments& InArgs, const TSharedRef<SDockTa
 	Settings->ConfigureVisLog();
 	Settings->LoadPersistentData();
 
-	ActionList.MapAction(Commands.StartRecording, FExecuteAction::CreateRaw(this, &SVisualLogger::HandleStartRecordingCommandExecute), FCanExecuteAction::CreateRaw(this, &SVisualLogger::HandleStartRecordingCommandCanExecute), FIsActionChecked(), FIsActionButtonVisible::CreateRaw(this, &SVisualLogger::HandleStartRecordingCommandIsVisible));
-	ActionList.MapAction(Commands.StopRecording, FExecuteAction::CreateRaw(this, &SVisualLogger::HandleStopRecordingCommandExecute), FCanExecuteAction::CreateRaw(this, &SVisualLogger::HandleStopRecordingCommandCanExecute), FIsActionChecked(), FIsActionButtonVisible::CreateRaw(this, &SVisualLogger::HandleStopRecordingCommandIsVisible));
-	ActionList.MapAction(Commands.Pause, FExecuteAction::CreateRaw(this, &SVisualLogger::HandlePauseCommandExecute), FCanExecuteAction::CreateRaw(this, &SVisualLogger::HandlePauseCommandCanExecute), FIsActionChecked(), FIsActionButtonVisible::CreateRaw(this, &SVisualLogger::HandlePauseCommandIsVisible));
-	ActionList.MapAction(Commands.Resume, FExecuteAction::CreateRaw(this, &SVisualLogger::HandleResumeCommandExecute), FCanExecuteAction::CreateRaw(this, &SVisualLogger::HandleResumeCommandCanExecute), FIsActionChecked(), FIsActionButtonVisible::CreateRaw(this, &SVisualLogger::HandleResumeCommandIsVisible));
-	ActionList.MapAction(Commands.LoadFromVLog, FExecuteAction::CreateRaw(this, &SVisualLogger::HandleLoadCommandExecute), FCanExecuteAction::CreateRaw(this, &SVisualLogger::HandleLoadCommandCanExecute), FIsActionChecked(), FIsActionButtonVisible::CreateRaw(this, &SVisualLogger::HandleLoadCommandCanExecute));
-	ActionList.MapAction(Commands.SaveToVLog, FExecuteAction::CreateRaw(this, &SVisualLogger::HandleSaveCommandExecute), FCanExecuteAction::CreateRaw(this, &SVisualLogger::HandleSaveCommandCanExecute), FIsActionChecked(), FIsActionButtonVisible::CreateRaw(this, &SVisualLogger::HandleSaveCommandCanExecute));
-	ActionList.MapAction(Commands.SaveAllToVLog, FExecuteAction::CreateRaw(this, &SVisualLogger::HandleSaveAllCommandExecute), FCanExecuteAction::CreateRaw(this, &SVisualLogger::HandleSaveCommandCanExecute), FIsActionChecked(), FIsActionButtonVisible::CreateRaw(this, &SVisualLogger::HandleSaveCommandCanExecute));
+	ActionList.MapAction(Commands.StartRecording, FExecuteAction::CreateSP(this, &SVisualLogger::HandleStartRecordingCommandExecute), FCanExecuteAction::CreateSP(this, &SVisualLogger::HandleStartRecordingCommandCanExecute), FIsActionChecked(), FIsActionButtonVisible::CreateSP(this, &SVisualLogger::HandleStartRecordingCommandIsVisible));
+	ActionList.MapAction(Commands.StopRecording, FExecuteAction::CreateSP(this, &SVisualLogger::HandleStopRecordingCommandExecute), FCanExecuteAction::CreateSP(this, &SVisualLogger::HandleStopRecordingCommandCanExecute), FIsActionChecked(), FIsActionButtonVisible::CreateSP(this, &SVisualLogger::HandleStopRecordingCommandIsVisible));
+	ActionList.MapAction(Commands.Pause, FExecuteAction::CreateSP(this, &SVisualLogger::HandlePauseCommandExecute), FCanExecuteAction::CreateSP(this, &SVisualLogger::HandlePauseCommandCanExecute), FIsActionChecked(), FIsActionButtonVisible::CreateSP(this, &SVisualLogger::HandlePauseCommandIsVisible));
+	ActionList.MapAction(Commands.Resume, FExecuteAction::CreateSP(this, &SVisualLogger::HandleResumeCommandExecute), FCanExecuteAction::CreateSP(this, &SVisualLogger::HandleResumeCommandCanExecute), FIsActionChecked(), FIsActionButtonVisible::CreateSP(this, &SVisualLogger::HandleResumeCommandIsVisible));
+	ActionList.MapAction(Commands.LoadFromVLog, FExecuteAction::CreateSP(this, &SVisualLogger::HandleLoadCommandExecute), FCanExecuteAction::CreateSP(this, &SVisualLogger::HandleLoadCommandCanExecute), FIsActionChecked(), FIsActionButtonVisible::CreateSP(this, &SVisualLogger::HandleLoadCommandCanExecute));
+	ActionList.MapAction(Commands.SaveToVLog, FExecuteAction::CreateSP(this, &SVisualLogger::HandleSaveCommandExecute), FCanExecuteAction::CreateSP(this, &SVisualLogger::HandleSaveCommandCanExecute), FIsActionChecked(), FIsActionButtonVisible::CreateSP(this, &SVisualLogger::HandleSaveCommandCanExecute));
+	ActionList.MapAction(Commands.SaveAllToVLog, FExecuteAction::CreateSP(this, &SVisualLogger::HandleSaveAllCommandExecute), FCanExecuteAction::CreateSP(this, &SVisualLogger::HandleSaveCommandCanExecute), FIsActionChecked(), FIsActionButtonVisible::CreateSP(this, &SVisualLogger::HandleSaveCommandCanExecute));
 	ActionList.MapAction(Commands.FreeCamera,
-		FExecuteAction::CreateRaw(this, &SVisualLogger::HandleCameraCommandExecute), 
-		FCanExecuteAction::CreateRaw(this, &SVisualLogger::HandleCameraCommandCanExecute), 
-		FIsActionChecked::CreateRaw(this, &SVisualLogger::HandleCameraCommandIsChecked),
-		FIsActionButtonVisible::CreateRaw(this, &SVisualLogger::HandleCameraCommandCanExecute));
+		FExecuteAction::CreateSP(this, &SVisualLogger::HandleCameraCommandExecute), 
+		FCanExecuteAction::CreateSP(this, &SVisualLogger::HandleCameraCommandCanExecute), 
+		FIsActionChecked::CreateSP(this, &SVisualLogger::HandleCameraCommandIsChecked),
+		FIsActionButtonVisible::CreateSP(this, &SVisualLogger::HandleCameraCommandCanExecute));
 	ActionList.MapAction(Commands.ToggleGraphs,
 		FExecuteAction::CreateLambda([](){bool& bEnableGraphsVisualization = ULogVisualizerSessionSettings::StaticClass()->GetDefaultObject<ULogVisualizerSessionSettings>()->bEnableGraphsVisualization; bEnableGraphsVisualization = !bEnableGraphsVisualization; }),
 		FCanExecuteAction::CreateLambda([this]()->bool{return FVisualLoggerGraphsDatabase::Get().ContainsHistogramGraphs(); }),
 		FIsActionChecked::CreateLambda([]()->bool{return ULogVisualizerSessionSettings::StaticClass()->GetDefaultObject<ULogVisualizerSessionSettings>()->bEnableGraphsVisualization; }),
 		FIsActionButtonVisible::CreateLambda([this]()->bool{return FVisualLoggerGraphsDatabase::Get().ContainsHistogramGraphs(); }));
 	ActionList.MapAction(Commands.ResetData, 
-		FExecuteAction::CreateRaw(this, &SVisualLogger::ResetData),
-		FCanExecuteAction::CreateRaw(this, &SVisualLogger::HandleSaveCommandCanExecute),
+		FExecuteAction::CreateSP(this, &SVisualLogger::ResetData),
+		FCanExecuteAction::CreateSP(this, &SVisualLogger::HandleSaveCommandCanExecute),
 		FIsActionChecked(), 
-		FIsActionButtonVisible::CreateRaw(this, &SVisualLogger::HandleSaveCommandCanExecute));
+		FIsActionButtonVisible::CreateSP(this, &SVisualLogger::HandleSaveCommandCanExecute));
 	ActionList.MapAction(Commands.Refresh,
-		FExecuteAction::CreateRaw(this, &SVisualLogger::HandleRefreshCommandExecute),
-		FCanExecuteAction::CreateRaw(this, &SVisualLogger::HandleRefreshCommandCanExecute),
+		FExecuteAction::CreateSP(this, &SVisualLogger::HandleRefreshCommandExecute),
+		FCanExecuteAction::CreateSP(this, &SVisualLogger::HandleRefreshCommandCanExecute),
 		FIsActionChecked(),
-		FIsActionButtonVisible::CreateRaw(this, &SVisualLogger::HandleRefreshCommandCanExecute));
+		FIsActionButtonVisible::CreateSP(this, &SVisualLogger::HandleRefreshCommandCanExecute));
+	ActionList.MapAction(Commands.AutoScroll,
+		FExecuteAction::CreateSP(this, &SVisualLogger::HandleAutoScrollCommandExecute),
+		FCanExecuteAction(),
+		FIsActionChecked::CreateSP(this, &SVisualLogger::HandleAutoScrollIsChecked),
+		FIsActionButtonVisible());
 
 
 	// Tab Spawners
 	TabManager = FGlobalTabmanager::Get()->NewTabManager(ConstructUnderMajorTab);
 	TSharedRef<FWorkspaceItem> AppMenuGroup = TabManager->AddLocalWorkspaceMenuCategory(LOCTEXT("VisualLoggerGroupName", "Visual Logger"));
 
-	TabManager->RegisterTabSpawner(ToolbarTabId, FOnSpawnTab::CreateRaw(this, &SVisualLogger::HandleTabManagerSpawnTab, ToolbarTabId))
+	TabManager->RegisterTabSpawner(ToolbarTabId, FOnSpawnTab::CreateSP(this, &SVisualLogger::HandleTabManagerSpawnTab, ToolbarTabId))
 		.SetDisplayName(LOCTEXT("ToolbarTabTitle", "Toolbar"))
 		.SetGroup(AppMenuGroup)
 		.SetIcon(FSlateIcon(FLogVisualizerStyle::Get().GetStyleSetName(), "ToolbarTabIcon"));
 
-	TabManager->RegisterTabSpawner(FiltersTabId, FOnSpawnTab::CreateRaw(this, &SVisualLogger::HandleTabManagerSpawnTab, FiltersTabId))
+	TabManager->RegisterTabSpawner(FiltersTabId, FOnSpawnTab::CreateSP(this, &SVisualLogger::HandleTabManagerSpawnTab, FiltersTabId))
 		.SetDisplayName(LOCTEXT("FiltersTabTitle", "Filters"))
 		.SetGroup(AppMenuGroup)
 		.SetIcon(FSlateIcon(FLogVisualizerStyle::Get().GetStyleSetName(), "FiltersTabIcon"));
 
-	TabManager->RegisterTabSpawner(MainViewTabId, FOnSpawnTab::CreateRaw(this, &SVisualLogger::HandleTabManagerSpawnTab, MainViewTabId))
+	TabManager->RegisterTabSpawner(MainViewTabId, FOnSpawnTab::CreateSP(this, &SVisualLogger::HandleTabManagerSpawnTab, MainViewTabId))
 		.SetDisplayName(LOCTEXT("MainViewTabTitle", "MainView"))
 		.SetGroup(AppMenuGroup)
 		.SetIcon(FSlateIcon(FLogVisualizerStyle::Get().GetStyleSetName(), "MainViewTabIcon"));
 
-	TabManager->RegisterTabSpawner(LogsListTabId, FOnSpawnTab::CreateRaw(this, &SVisualLogger::HandleTabManagerSpawnTab, LogsListTabId))
+	TabManager->RegisterTabSpawner(LogsListTabId, FOnSpawnTab::CreateSP(this, &SVisualLogger::HandleTabManagerSpawnTab, LogsListTabId))
 		.SetDisplayName(LOCTEXT("LogsListTabTitle", "LogsList"))
 		.SetGroup(AppMenuGroup)
 		.SetIcon(FSlateIcon(FLogVisualizerStyle::Get().GetStyleSetName(), "LogsListTabIcon"));
 
-	TabManager->RegisterTabSpawner(StatusViewTabId, FOnSpawnTab::CreateRaw(this, &SVisualLogger::HandleTabManagerSpawnTab, StatusViewTabId))
+	TabManager->RegisterTabSpawner(StatusViewTabId, FOnSpawnTab::CreateSP(this, &SVisualLogger::HandleTabManagerSpawnTab, StatusViewTabId))
 		.SetDisplayName(LOCTEXT("StatusViewTabTitle", "StatusView"))
 		.SetGroup(AppMenuGroup)
 		.SetIcon(FSlateIcon(FLogVisualizerStyle::Get().GetStyleSetName(), "StatusViewTabIcon"));
@@ -322,7 +331,7 @@ void SVisualLogger::Construct(const FArguments& InArgs, const TSharedRef<SDockTa
 					)
 			)
 		);
-	TabManager->SetOnPersistLayout(FTabManager::FOnPersistLayout::CreateRaw(this, &SVisualLogger::HandleTabManagerPersistLayout));
+	TabManager->SetOnPersistLayout(FTabManager::FOnPersistLayout::CreateSP(this, &SVisualLogger::HandleTabManagerPersistLayout));
 
 	// Window Menu
 	FMenuBarBuilder MenuBarBuilder = FMenuBarBuilder(TSharedPtr<FUICommandList>());
@@ -382,11 +391,6 @@ void SVisualLogger::OnNewLogEntry(const FVisualLogDevice::FVisualLogEntryItem& E
 	}
 
 	FVisualLoggerDatabase::Get().AddItem(Entry);
-
-	if (ULogVisualizerSettings::StaticClass()->GetDefaultObject<ULogVisualizerSettings>()->bStickToRecentData)
-	{
-		FLogVisualizer::Get().GetTimeSliderController()->CommitScrubPosition(Entry.Entry.TimeStamp, false);
-	}
 }
 
 void SVisualLogger::HandleMajorTabPersistVisualState()
@@ -603,6 +607,21 @@ bool SVisualLogger::HandleLoadCommandCanExecute() const
 	return true;
 }
 
+void SVisualLogger::HandleAutoScrollCommandExecute()
+{
+	bAutoScrollToLastItem = !bAutoScrollToLastItem;
+	if (bAutoScrollToLastItem)
+	{
+		TGuardValue<bool> Guard(bAutoScrollingToLastItem, true);
+		FLogVisualizer::Get().GotoLastItemAnyRow();
+	}
+}
+
+bool SVisualLogger::HandleAutoScrollIsChecked() const
+{
+	return bAutoScrollToLastItem;
+}
+
 void SVisualLogger::HandleLoadCommandExecute()
 {
 	FArchive Ar;
@@ -776,9 +795,9 @@ void SVisualLogger::ResetData()
 		}
 	}
 
-	FLogVisualizer::Get().GetEvents().OnLogLineSelectionChanged = FOnLogLineSelectionChanged::CreateRaw(this, &SVisualLogger::OnLogLineSelectionChanged);
-	FLogVisualizer::Get().GetEvents().OnKeyboardEvent = FOnKeyboardEvent::CreateRaw(this, &SVisualLogger::OnKeyboardRedirection);
-	FLogVisualizer::Get().GetTimeSliderController().Get()->GetTimeSliderArgs().OnScrubPositionChanged = FVisualLoggerTimeSliderArgs::FOnScrubPositionChanged::CreateRaw(this, &SVisualLogger::OnScrubPositionChanged);
+	FLogVisualizer::Get().GetEvents().OnLogLineSelectionChanged = FOnLogLineSelectionChanged::CreateSP(this, &SVisualLogger::OnLogLineSelectionChanged);
+	FLogVisualizer::Get().GetEvents().OnKeyboardEvent = FOnKeyboardEvent::CreateSP(this, &SVisualLogger::OnKeyboardRedirection);
+	FLogVisualizer::Get().GetTimeSliderController().Get()->GetTimeSliderArgs().OnScrubPositionChanged = FVisualLoggerTimeSliderArgs::FOnScrubPositionChanged::CreateSP(this, &SVisualLogger::OnScrubPositionChanged);
 }
 
 void SVisualLogger::CollectFilterVolumes()
@@ -1038,6 +1057,12 @@ void SVisualLogger::OnNewItemHandler(const FVisualLoggerDBRow& DBRow, int32 Item
 	// In such case, we might need to update row visibility/
 	// Note that this is not called within UpdateVisibilityForEntry since it can be called by async tasks.
 	FVisualLoggerDatabase::Get().SetRowVisibility(DBRow.GetOwnerName(), DBRow.GetNumberOfHiddenItems() != DBRow.GetItems().Num());
+
+	if (bAutoScrollToLastItem)
+	{
+		TGuardValue<bool> Guard(bAutoScrollingToLastItem, true);
+		FLogVisualizer::Get().GotoLastItemAnyRow();
+	}
 }
 
 void SVisualLogger::UpdateVisibilityForEntry(const FVisualLoggerDBRow& DBRow, int32 ItemIndex)
@@ -1143,6 +1168,12 @@ void SVisualLogger::OnScrubPositionChanged(double NewScrubPosition, bool bScrubb
 			Extension->OnScrubPositionChanged(FVisualLoggerEditorInterface::Get(), NewScrubPosition, bScrubbing);
 		}
 	}
+
+	// Cancel auto-scrolling on scrub (except if it comes from within and the actual bAutoScroll behavior, which is guarded by bAutoScrollingToLastItem)
+	if (!bAutoScrollingToLastItem)
+	{
+		bAutoScrollToLastItem = false;
+	}
 }
 
 FReply SVisualLogger::OnKeyboardRedirection(const FGeometry& MyGeometry, const FKeyEvent& InKeyEvent)
diff --git a/Engine/Source/Developer/LogVisualizer/Private/SVisualLogger.h b/Engine/Source/Developer/LogVisualizer/Private/SVisualLogger.h
--- a/Engine/Source/Developer/LogVisualizer/Private/SVisualLogger.h
+++ b/Engine/Source/Developer/LogVisualizer/Private/SVisualLogger.h
@@ -85,6 +85,9 @@ public:
 	bool HandleLoadCommandCanExecute() const;
 	void HandleLoadCommandExecute();
 
+	void HandleAutoScrollCommandExecute();
+	bool HandleAutoScrollIsChecked() const;
+
 	bool HandleSaveCommandCanExecute() const;
 	void HandleSaveCommandExecute();
 
@@ -138,6 +141,11 @@ protected:
 
 	bool bGotHistogramData;
 
+	// Instructs the visual logger to always scrub to the last item
+	bool bAutoScrollToLastItem;
+	// Guard value that prevents canceling bAutoScrollToLastItem when we're actually auto-scrubbing (instead of manually scrubbing)
+	mutable bool bAutoScrollingToLastItem;
+
 #if WITH_EDITOR
 	FDelegateHandle PostWorldInitializationHandle;
 #endif
diff --git a/Engine/Source/Developer/LogVisualizer/Private/SVisualLoggerToolbar.cpp b/Engine/Source/Developer/LogVisualizer/Private/SVisualLoggerToolbar.cpp
--- a/Engine/Source/Developer/LogVisualizer/Private/SVisualLoggerToolbar.cpp
+++ b/Engine/Source/Developer/LogVisualizer/Private/SVisualLoggerToolbar.cpp
@@ -46,6 +46,7 @@ TSharedRef<SWidget> SVisualLoggerToolbar::MakeToolbar(const TSharedRef<FUIComman
 		ToolBarBuilder.AddToolBarButton(FVisualLoggerCommands::Get().ResetData, NAME_None, LOCTEXT("ResetData", "Clear"), TAttribute<FText>(), FSlateIcon(FLogVisualizerStyle::Get().GetStyleSetName(), TEXT("Toolbar.Remove")));
 		ToolBarBuilder.AddToolBarButton(FVisualLoggerCommands::Get().ToggleGraphs, NAME_None, LOCTEXT("ToggleGraphs", "Graphs"), TAttribute<FText>(), FSlateIcon(FLogVisualizerStyle::Get().GetStyleSetName(), TEXT("Toolbar.Graphs")));
 		ToolBarBuilder.AddToolBarButton(FVisualLoggerCommands::Get().Refresh, NAME_None, LOCTEXT("ForceRefresh", "Refresh"), TAttribute<FText>(), FSlateIcon(FAppStyle::GetAppStyleSetName(), "Icons.Refresh"));
+		ToolBarBuilder.AddToolBarButton(FVisualLoggerCommands::Get().AutoScroll, NAME_None, LOCTEXT("AutoScroll", "Auto-Scroll"), TAttribute<FText>(), FSlateIcon(FLogVisualizerStyle::Get().GetStyleSetName(), "Toolbar.AutoScroll"));
 
 		ToolBarBuilder.AddSeparator();
 		ToolBarBuilder.AddToolBarWidget(MakeClippingDistanceWidget(), LOCTEXT("FarClippingDistance", "Far Clipping Distance"), /*InTutorialHighlightName = */NAME_None, /*bInSearchable = */true, LOCTEXT("FarClipDistanceTooltip", "Max distance after which visual log items will stop being rendered (ignored if <=0). This can help performance/visibility in situations where lots of visual logs are displayed on screen."));
diff --git a/Engine/Source/Developer/LogVisualizer/Private/VisualLoggerCommands.h b/Engine/Source/Developer/LogVisualizer/Private/VisualLoggerCommands.h
--- a/Engine/Source/Developer/LogVisualizer/Private/VisualLoggerCommands.h
+++ b/Engine/Source/Developer/LogVisualizer/Private/VisualLoggerCommands.h
@@ -34,6 +34,7 @@ public:
 		UI_COMMAND(ToggleGraphs, "ToggleGraphs", "Toggle graphs visualization on/off", EUserInterfaceActionType::ToggleButton, FInputChord());
 		UI_COMMAND(ResetData, "Clear", "Clear all data", EUserInterfaceActionType::Button, FInputChord());
 		UI_COMMAND(Refresh, "Refresh", "Refresh UI from memory buffer", EUserInterfaceActionType::Button, FInputChord());
+		UI_COMMAND(AutoScroll, "AutoScroll", "Maintains the cursor on the last recorded event", EUserInterfaceActionType::ToggleButton, FInputChord());
 		UI_COMMAND(HideEmpty, "HideEmpty", "Hide logs without text information", EUserInterfaceActionType::ToggleButton, FInputChord());
 	}
 
@@ -50,6 +51,7 @@ public:
 	TSharedPtr<FUICommandInfo> ToggleGraphs;
 	TSharedPtr<FUICommandInfo> ResetData;
 	TSharedPtr<FUICommandInfo> Refresh;
+	TSharedPtr<FUICommandInfo> AutoScroll;
 	TSharedPtr<FUICommandInfo> HideEmpty;
 };
 
diff --git a/Engine/Source/Developer/LogVisualizer/Public/LogVisualizerPublic.h b/Engine/Source/Developer/LogVisualizer/Public/LogVisualizerPublic.h
--- a/Engine/Source/Developer/LogVisualizer/Public/LogVisualizerPublic.h
+++ b/Engine/Source/Developer/LogVisualizer/Public/LogVisualizerPublic.h
@@ -62,6 +62,7 @@ struct LOGVISUALIZER_API FLogVisualizer
 	void GotoPreviousItem(FName RowName, int32 MoveDistance = 1);
 	void GotoFirstItem(FName RowName);
 	void GotoLastItem(FName RowName);
+	void GotoLastItemAnyRow();
 
 	void UpdateCameraPosition(FName Rowname, int32 ItemIndes);
 
diff --git a/Engine/Source/Developer/LogVisualizer/Public/LogVisualizerSettings.h b/Engine/Source/Developer/LogVisualizer/Public/LogVisualizerSettings.h
--- a/Engine/Source/Developer/LogVisualizer/Public/LogVisualizerSettings.h
+++ b/Engine/Source/Developer/LogVisualizer/Public/LogVisualizerSettings.h
@@ -125,8 +125,8 @@ public:
 	UPROPERTY(EditAnywhere, config, Category = "VisualLogger", meta = (EditCondition = "bIgnoreTrivialLogs", ClampMin = "0", ClampMax = "10", UIMin = "0", UIMax = "10"))
 	int32 TrivialLogsThreshold;
 
-	/**Whether to show the recent data or not. Property disabled for now.*/
-	UPROPERTY(EditAnywhere, config, Category = "VisualLogger")
+	UE_DEPRECATED(5.6, "This is now controlled by the auto-scroll button in the tool")
+	UPROPERTY()
 	bool bStickToRecentData;
 
 	/**Whether to reset current data or not for each new session.*/
