[VisualLogger] improved filtering for row items and tooltips
- replaced MatchCategoryFilters by faster ShouldDisplayCategory base on FName comparison
- changed selected entry color to reduce confusion with the color used when warnings are presents
- error and warning colored boxes are now slightly larger if they are the selected entry to preserve that information
- filtered rows are now preserved when a category filter level is modified
- tooltips now reflect the category filters

diff --git a/Engine/Source/Developer/LogVisualizer/Private/LogVisualizerSettings.cpp b/Engine/Source/Developer/LogVisualizer/Private/LogVisualizerSettings.cpp
--- a/Engine/Source/Developer/LogVisualizer/Private/LogVisualizerSettings.cpp
+++ b/Engine/Source/Developer/LogVisualizer/Private/LogVisualizerSettings.cpp
@@ -245,24 +245,57 @@ const TArray<FString>& FVisualLoggerFilters::GetSelectedObjects() const
 	return SelectedClasses;
 }
 
-// @todo both MatchCategoryFilters and MatchSearchString function names are not clear enough
-bool FVisualLoggerFilters::MatchCategoryFilters(FString String, ELogVerbosity::Type Verbosity)
+bool FVisualLoggerFilters::ShouldDisplayLine(const FVisualLogLine& Line, const bool bSearchInsideLogs) const
 {
-	ULogVisualizerSettings* Settings = ULogVisualizerSettings::StaticClass()->GetDefaultObject<ULogVisualizerSettings>();
+	bool bShowLine = ShouldDisplayCategory(Line.Category, Line.Verbosity);
+	if (bShowLine && bSearchInsideLogs)
+	{
+		const FString SearchString = GetSearchString();
+		if (SearchString.Len() > 0
+			&& Line.Line.Find(SearchString) == INDEX_NONE
+			&& Line.Category.ToString().Find(SearchString) == INDEX_NONE)
+		{
+			bShowLine = false;
+		}
+	}
+	return bShowLine;
+}
 
+bool FVisualLoggerFilters::ShouldDisplayCategory(const FName Name, const ELogVerbosity::Type Verbosity) const
+{
+	const ULogVisualizerSettings* Settings = ULogVisualizerSettings::StaticClass()->GetDefaultObject<ULogVisualizerSettings>();
 	for (const FCategoryFilter& Filter : Categories)
 	{
-		if (Filter.CategoryName.Equals(String, ESearchCase::IgnoreCase))
+		if (Filter.CategoryName == Name)
 		{
 			if (Filter.Enabled)
 			{
 				const bool bFineWithSearchString = Settings->bSearchInsideLogs == true || (SearchBoxFilter.Len() == 0 || Filter.CategoryName.Find(SearchBoxFilter) != INDEX_NONE);
 				return bFineWithSearchString && Verbosity <= Filter.LogVerbosity;
 			}
-			else
+
+			return false;
+		}
+	}
+
+	// if filter for a given log category has not been found we let it be, we allow it.
+	return true;
+}
+
+bool FVisualLoggerFilters::ShouldDisplayCategoryByString(const FStringView String, const ELogVerbosity::Type Verbosity) const
+{
+	const ULogVisualizerSettings* Settings = ULogVisualizerSettings::StaticClass()->GetDefaultObject<ULogVisualizerSettings>();
+	for (const FCategoryFilter& Filter : Categories)
+	{
+		if (Filter.CategoryName.Equals(String.GetData(), ESearchCase::IgnoreCase))
+		{
+			if (Filter.Enabled)
 			{
-				return false;
+				const bool bFineWithSearchString = Settings->bSearchInsideLogs == true || (SearchBoxFilter.Len() == 0 || Filter.CategoryName.Find(SearchBoxFilter) != INDEX_NONE);
+				return bFineWithSearchString && Verbosity <= Filter.LogVerbosity;
 			}
+
+			return false;
 		}
 	}
 
@@ -274,7 +307,7 @@ void FVisualLoggerFilters::DeactivateAllButThis(const FString& InName)
 {
 	for (FCategoryFilter& Filter : Categories)
 	{
-		Filter.Enabled = Filter.CategoryName != InName ? false : true;
+		Filter.Enabled = Filter.CategoryName == InName;
 	}
 }
 
diff --git a/Engine/Source/Developer/LogVisualizer/Private/LogVisualizerStyle.cpp b/Engine/Source/Developer/LogVisualizer/Private/LogVisualizerStyle.cpp
--- a/Engine/Source/Developer/LogVisualizer/Private/LogVisualizerStyle.cpp
+++ b/Engine/Source/Developer/LogVisualizer/Private/LogVisualizerStyle.cpp
@@ -67,7 +67,7 @@ TSharedRef< FSlateStyleSet > FLogVisualizerStyle::Create()
 		Style.Set("LogVisualizerApp.TabIcon", new IMAGE_BRUSH_SVG("Starship/Common/VisualLogger", Icon16x16));
 		
 		Style.Set("LogVisualizer.LogBar.Background", new BOX_BRUSH("Common/ProgressBar_Background", FMargin(5.f / 12.f)));
-		Style.Set("LogVisualizer.LogBar.Selected", new CORE_BOX_BRUSH("Common/ProfileVisualizer_Selected", FMargin(5.f / 12.f)));
+		Style.Set("LogVisualizer.LogBar.Selected", new CORE_BOX_BRUSH("Common/ProfileVisualizer_Normal", FMargin(5.f / 12.f)));
 		Style.Set("LogVisualizer.LogBar.EntryDefault", new CORE_BOX_BRUSH("Common/ProfileVisualizer_Mono", FMargin(5.f / 12.f)));
 		Style.Set("LogVisualizer.LogBar.TimeMark", new BOX_BRUSH("Icons/LV_BarMark", FMargin(5.f / 12.f)));
 
diff --git a/Engine/Source/Developer/LogVisualizer/Private/SVisualLogger.cpp b/Engine/Source/Developer/LogVisualizer/Private/SVisualLogger.cpp
--- a/Engine/Source/Developer/LogVisualizer/Private/SVisualLogger.cpp
+++ b/Engine/Source/Developer/LogVisualizer/Private/SVisualLogger.cpp
@@ -126,8 +126,8 @@ SVisualLogger::~SVisualLogger()
 	InternalDevice.Reset();
 
 #if WITH_EDITOR
-	ULogVisualizerSettings::StaticClass()->GetDefaultObject<ULogVisualizerSettings>()->SavePersistentData();
-	ULogVisualizerSettings::StaticClass()->GetDefaultObject<ULogVisualizerSettings>()->OnSettingChanged().RemoveAll(this);
+	GetMutableDefault<ULogVisualizerSettings>()->SavePersistentData();
+	GetMutableDefault<ULogVisualizerSettings>()->OnSettingChanged().RemoveAll(this);
 #endif
 
 	if (LastUsedWorld.IsValid())
@@ -205,7 +205,7 @@ void SVisualLogger::Construct(const FArguments& InArgs, const TSharedRef<SDockTa
 		}
 	});
 
-	ULogVisualizerSettings::StaticClass()->GetDefaultObject<ULogVisualizerSettings>()->OnSettingChanged().AddRaw(this, &SVisualLogger::OnSettingsChanged);
+	GetMutableDefault<ULogVisualizerSettings>()->OnSettingChanged().AddRaw(this, &SVisualLogger::OnSettingsChanged);
 #endif
 
 	GEngine->OnWorldAdded().AddRaw(this, &SVisualLogger::OnNewWorld);
@@ -219,7 +219,7 @@ void SVisualLogger::Construct(const FArguments& InArgs, const TSharedRef<SDockTa
 	const FVisualLoggerCommands& Commands = FVisualLoggerCommands::Get();
 	FUICommandList& ActionList = *CommandList;
 
-	ULogVisualizerSettings* Settings = ULogVisualizerSettings::StaticClass()->GetDefaultObject<ULogVisualizerSettings>();
+	ULogVisualizerSettings* Settings = GetMutableDefault<ULogVisualizerSettings>();
 	Settings->ConfigureVisLog();
 	Settings->LoadPersistentData();
 
@@ -480,7 +480,7 @@ void SVisualLogger::HandleStopRecordingCommandExecute()
 {
 	UWorld* World = FLogVisualizer::Get().GetWorld();
 
-	if (FParse::Param(FCommandLine::Get(), TEXT("LogNavOctree")) == true && ULogVisualizerSettings::StaticClass()->GetDefaultObject<ULogVisualizerSettings>()->bLogNavOctreeOnStop)
+	if (FParse::Param(FCommandLine::Get(), TEXT("LogNavOctree")) == true && GetDefault<ULogVisualizerSettings>()->bLogNavOctreeOnStop)
 	{
 		FVisualLogger::NavigationDataDump(World, LogNavigation, ELogVerbosity::Log, FBox());
 	}
@@ -511,7 +511,7 @@ bool SVisualLogger::HandlePauseCommandCanExecute() const
 
 void SVisualLogger::HandlePauseCommandExecute()
 {
-	if (ULogVisualizerSettings::StaticClass()->GetDefaultObject<ULogVisualizerSettings>()->bUsePlayersOnlyForPause)
+	if (GetDefault<ULogVisualizerSettings>()->bUsePlayersOnlyForPause)
 	{
 		const TIndirectArray<FWorldContext>& WorldContexts = GEngine->GetWorldContexts();
 		for (const FWorldContext& Context : WorldContexts)
@@ -538,7 +538,7 @@ bool SVisualLogger::HandleResumeCommandCanExecute() const
 
 void SVisualLogger::HandleResumeCommandExecute()
 {
-	if (ULogVisualizerSettings::StaticClass()->GetDefaultObject<ULogVisualizerSettings>()->bUsePlayersOnlyForPause)
+	if (GetDefault<ULogVisualizerSettings>()->bUsePlayersOnlyForPause)
 	{
 		const TIndirectArray<FWorldContext>& WorldContexts = GEngine->GetWorldContexts();
 		for (const FWorldContext& Context : WorldContexts)
@@ -871,7 +871,7 @@ void SVisualLogger::OnNewWorld(UWorld* NewWorld)
 
 	// Only reset data when activating a new game world, not when returning to the Editor world (i.e. after PIE)
 	if (!IsValid(NewWorld) ||
-		(NewWorld->IsGameWorld() && ULogVisualizerSettings::StaticClass()->GetDefaultObject<ULogVisualizerSettings>()->bResetDataWithNewSession))
+		(NewWorld->IsGameWorld() && GetDefault<ULogVisualizerSettings>()->bResetDataWithNewSession))
 	{
 		ResetData();
 	}
@@ -923,7 +923,6 @@ void SVisualLogger::OnItemsSelectionChanged(const FVisualLoggerDBRow& ChangedRow
 void SVisualLogger::OnFiltersChanged()
 {
 	const uint32 StartCycles = FPlatformTime::Cycles();
-	const FString QuickSearchStrng = FVisualLoggerFilters::Get().GetSearchString();
 
 	TArray<TFuture<void> > AllFutures;
 	for (auto Iterator = FVisualLoggerDatabase::Get().GetRowIterator(); Iterator; ++Iterator)
@@ -962,7 +961,7 @@ void SVisualLogger::OnFiltersChanged()
 	for (auto Iterator = FVisualLoggerDatabase::Get().GetRowIterator(); Iterator; ++Iterator)
 	{
 		FVisualLoggerDBRow& DBRow = *Iterator;
-		FVisualLoggerDatabase::Get().SetRowVisibility(DBRow.GetOwnerName(), DBRow.GetNumberOfHiddenItems() != DBRow.GetItems().Num());
+		FVisualLoggerDatabase::Get().SetRowVisibility(DBRow.GetOwnerName(), DBRow.IsRowVisible());
 	}
 
 	const int32 BlockingCycles = int32(FPlatformTime::Cycles() - StartCycles);
@@ -991,8 +990,6 @@ void SVisualLogger::OnFiltersSearchChanged(const FText& Filter)
 
 	FVisualLoggerFilters::Get().SetSearchString(Filter.ToString());
 
-	const FString QuickSearchStrng = FVisualLoggerFilters::Get().GetSearchString();
-
 	TArray<TFuture<void> > AllFutures;
 	for (auto Iterator = FVisualLoggerDatabase::Get().GetRowIterator(); Iterator; ++Iterator)
 	{
@@ -1031,7 +1028,7 @@ void SVisualLogger::OnFiltersSearchChanged(const FText& Filter)
 	for (auto Iterator = FVisualLoggerDatabase::Get().GetRowIterator(); Iterator; ++Iterator)
 	{
 		FVisualLoggerDBRow& DBRow = *Iterator;
-		FVisualLoggerDatabase::Get().SetRowVisibility(DBRow.GetOwnerName(), DBRow.GetNumberOfHiddenItems() != DBRow.GetItems().Num());
+		FVisualLoggerDatabase::Get().SetRowVisibility(DBRow.GetOwnerName(), DBRow.IsRowVisible());
 	}
 
 	if (LogsList.IsValid())
@@ -1056,7 +1053,7 @@ void SVisualLogger::OnNewItemHandler(const FVisualLoggerDBRow& DBRow, int32 Item
 	// A new item can be hidden by the category filters by UpdateVisibilityForEntry.
 	// In such case, we might need to update row visibility/
 	// Note that this is not called within UpdateVisibilityForEntry since it can be called by async tasks.
-	FVisualLoggerDatabase::Get().SetRowVisibility(DBRow.GetOwnerName(), DBRow.GetNumberOfHiddenItems() != DBRow.GetItems().Num());
+	FVisualLoggerDatabase::Get().SetRowVisibility(DBRow.GetOwnerName(), DBRow.IsRowVisible());
 
 	if (bAutoScrollToLastItem)
 	{
@@ -1067,25 +1064,23 @@ void SVisualLogger::OnNewItemHandler(const FVisualLoggerDBRow& DBRow, int32 Item
 
 void SVisualLogger::UpdateVisibilityForEntry(const FVisualLoggerDBRow& DBRow, int32 ItemIndex)
 {
-	ULogVisualizerSettings* Settings = ULogVisualizerSettings::StaticClass()->GetDefaultObject<ULogVisualizerSettings>();
+	const ULogVisualizerSettings* Settings = GetDefault<ULogVisualizerSettings>();
 	const FVisualLogDevice::FVisualLogEntryItem& CurrentEntry = DBRow.GetItems()[ItemIndex];
-	FString SearchString = FVisualLoggerFilters::Get().GetSearchString();
+	const FString SearchString = FVisualLoggerFilters::Get().GetSearchString();
 
 	TArray<FVisualLoggerCategoryVerbosityPair> OutCategories;
 	FVisualLoggerHelpers::GetCategories(CurrentEntry.Entry, OutCategories);
 
-	bool bPassingFilter = true;
+	bool bPassingFilter = false;
 
-	bool bHasValidCategories = false;
 	for (FVisualLoggerCategoryVerbosityPair& CategoryPair : OutCategories)
 	{
-		if (FVisualLoggerFilters::Get().MatchCategoryFilters(CategoryPair.CategoryName.ToString(), CategoryPair.Verbosity))
+		if (FVisualLoggerFilters::Get().ShouldDisplayCategory(CategoryPair.CategoryName, CategoryPair.Verbosity))
 		{
-			bHasValidCategories = true;
+			bPassingFilter = true;
 			break;
 		}
 	}
-	bPassingFilter &= bHasValidCategories;
 
 	if (bPassingFilter && Settings->bUseFilterVolumes && FilterBoxes.Num() > 0 && CurrentEntry.Entry.bIsLocationValid)
 	{
@@ -1104,28 +1099,27 @@ void SVisualLogger::UpdateVisibilityForEntry(const FVisualLoggerDBRow& DBRow, in
 
 	if (bPassingFilter && Settings->bSearchInsideLogs && SearchString.Len() > 0)
 	{
-		bool bMatchSearchString = false;
+		bPassingFilter = false;
 		for (const FVisualLogLine& CurrentLine : CurrentEntry.Entry.LogLines)
 		{
 			if (CurrentLine.Line.Find(SearchString) != INDEX_NONE || CurrentLine.Category.ToString().Find(SearchString) != INDEX_NONE)
 			{
-				bMatchSearchString = true;
+				bPassingFilter = true;
 				break;
 			}
 		}
-		if (!bMatchSearchString)
+
+		if (!bPassingFilter)
 		{
 			for (const FVisualLogEvent& CurrentEvent : CurrentEntry.Entry.Events)
 			{
 				if (CurrentEvent.Name.Find(SearchString) != INDEX_NONE)
 				{
-					bMatchSearchString = true;
+					bPassingFilter = true;
 					break;
 				}
 			}
 		}
-
-		bPassingFilter &= bMatchSearchString;
 	}
 
 	FVisualLoggerDatabase::Get().GetRowByName(DBRow.GetOwnerName()).SetItemVisibility(ItemIndex, bPassingFilter);
diff --git a/Engine/Source/Developer/LogVisualizer/Private/SVisualLoggerLogsList.cpp b/Engine/Source/Developer/LogVisualizer/Private/SVisualLoggerLogsList.cpp
--- a/Engine/Source/Developer/LogVisualizer/Private/SVisualLoggerLogsList.cpp
+++ b/Engine/Source/Developer/LogVisualizer/Private/SVisualLoggerLogsList.cpp
@@ -84,7 +84,7 @@ public:
 	FText GetFilterText() const
 	{
 		static FText NoText;
-		const bool bSearchInsideLogs = ULogVisualizerSettings::StaticClass()->GetDefaultObject<ULogVisualizerSettings>()->bSearchInsideLogs;
+		const bool bSearchInsideLogs = GetDefault<ULogVisualizerSettings>()->bSearchInsideLogs;
 		return bSearchInsideLogs ? FText::FromString(FVisualLoggerFilters::Get().GetSearchString()) : NoText;
 	}
 
@@ -169,13 +169,6 @@ TSharedRef<ITableRow> SVisualLoggerLogsList::LogEntryLinesGenerateRow(TSharedPtr
 	return SNew(SVisualLoggerTableRow, OwnerTable, Item);
 }
 
-FText SVisualLoggerLogsList::GetFilterText() const
-{
-	static FText NoText;
-	const bool bSearchInsideLogs = ULogVisualizerSettings::StaticClass()->GetDefaultObject<ULogVisualizerSettings>()->bSearchInsideLogs;
-	return bSearchInsideLogs ? FText::FromString(FVisualLoggerFilters::Get().GetSearchString()) : NoText;
-}
-
 void SVisualLoggerLogsList::OnFiltersSearchChanged(const FText& Filter)
 {
 	OnFiltersChanged();
@@ -255,13 +248,17 @@ void SVisualLoggerLogsList::GenerateLogs(const FVisualLogDevice::FVisualLogEntry
 {
 	TArray<FVisualLoggerCategoryVerbosityPair> OutCategories;
 	FVisualLoggerHelpers::GetCategories(LogEntry.Entry, OutCategories);
-	bool bHasValidCategory = false;
-	for (auto& CurrentCategory : OutCategories)
+	bool bHasDisplayedCategory = false;
+	for (const FVisualLoggerCategoryVerbosityPair& CurrentCategory : OutCategories)
 	{
-		bHasValidCategory |= FVisualLoggerFilters::Get().MatchCategoryFilters(CurrentCategory.CategoryName.ToString(), CurrentCategory.Verbosity);
+		if (FVisualLoggerFilters::Get().ShouldDisplayCategory(CurrentCategory.CategoryName, CurrentCategory.Verbosity))
+		{
+			bHasDisplayedCategory = true;
+			break;
+		}
 	}
 	
-	if (!bHasValidCategory)
+	if (!bHasDisplayedCategory)
 	{
 		return;
 	}
@@ -279,21 +276,10 @@ void SVisualLoggerLogsList::GenerateLogs(const FVisualLogDevice::FVisualLogEntry
 	}
 
 	const FVisualLogLine* LogLine = LogEntry.Entry.LogLines.GetData();
-	const bool bSearchInsideLogs = ULogVisualizerSettings::StaticClass()->GetDefaultObject<ULogVisualizerSettings>()->bSearchInsideLogs;
+	const bool bSearchInsideLogs = GetDefault<ULogVisualizerSettings>()->bSearchInsideLogs;
 	for (int LineIndex = 0; LineIndex < LogEntry.Entry.LogLines.Num(); ++LineIndex, ++LogLine)
 	{
-		bool bShowLine = FVisualLoggerFilters::Get().MatchCategoryFilters(LogLine->Category.ToString(), LogLine->Verbosity);
-		if (bSearchInsideLogs)
-		{
-			FString String = FVisualLoggerFilters::Get().GetSearchString();
-			if (String.Len() > 0)
-			{
-				bShowLine &= LogLine->Line.Find(String) != INDEX_NONE || LogLine->Category.ToString().Find(String) != INDEX_NONE;
-			}
-		}
-		
-
-		if (bShowLine)
+		if (FVisualLoggerFilters::Get().ShouldDisplayLine(*LogLine, bSearchInsideLogs))
 		{
 			FLogEntryItem EntryItem;
 			EntryItem.Category = LogLine->Category.ToString();
@@ -309,9 +295,7 @@ void SVisualLoggerLogsList::GenerateLogs(const FVisualLogDevice::FVisualLogEntry
 
 	for (auto& Event : LogEntry.Entry.Events)
 	{
-		bool bShowLine = FVisualLoggerFilters::Get().MatchCategoryFilters(Event.Name, Event.Verbosity);
-
-		if (bShowLine)
+		if (FVisualLoggerFilters::Get().ShouldDisplayCategoryByString(Event.Name, Event.Verbosity))
 		{
 			FLogEntryItem EntryItem;
 			EntryItem.Category = Event.Name;
diff --git a/Engine/Source/Developer/LogVisualizer/Private/SVisualLoggerLogsList.h b/Engine/Source/Developer/LogVisualizer/Private/SVisualLoggerLogsList.h
--- a/Engine/Source/Developer/LogVisualizer/Private/SVisualLoggerLogsList.h
+++ b/Engine/Source/Developer/LogVisualizer/Private/SVisualLoggerLogsList.h
@@ -2,7 +2,6 @@
 
 #pragma once
 
-#include "CoreMinimal.h"
 #include "Input/Reply.h"
 #include "Widgets/DeclarativeSyntaxSupport.h"
 #include "Framework/Commands/UICommandList.h"
@@ -32,7 +31,6 @@ public:
 	void OnFiltersSearchChanged(const FText& Filter);
 
 protected:
-	FText GetFilterText() const;
 	void RegenerateLogEntries();
 	void GenerateLogs(const FVisualLogDevice::FVisualLogEntryItem& EntryItem, bool bGenerateHeader);
 	void LogEntryLineSelectionChanged(TSharedPtr<FLogEntryItem> SelectedItem, ESelectInfo::Type SelectInfo);
diff --git a/Engine/Source/Developer/LogVisualizer/Private/SVisualLoggerTimeline.cpp b/Engine/Source/Developer/LogVisualizer/Private/SVisualLoggerTimeline.cpp
--- a/Engine/Source/Developer/LogVisualizer/Private/SVisualLoggerTimeline.cpp
+++ b/Engine/Source/Developer/LogVisualizer/Private/SVisualLoggerTimeline.cpp
@@ -65,17 +65,19 @@ bool SLogVisualizerTimeline::IsSelected() const
 
 void SLogVisualizerTimeline::UpdateVisibility()
 {
-	ULogVisualizerSettings* Settings = ULogVisualizerSettings::StaticClass()->GetDefaultObject<ULogVisualizerSettings>();
+	const ULogVisualizerSettings* Settings = GetDefault<ULogVisualizerSettings>();
 
-	if (FVisualLoggerDatabase::Get().ContainsRowByName(GetName()))
+	const FName RowName = GetName();
+	if (FVisualLoggerDatabase::Get().ContainsRowByName(RowName))
 	{
-		FVisualLoggerDBRow& DataRow = FVisualLoggerDatabase::Get().GetRowByName(GetName());
+		FVisualLoggerDBRow& DataRow = FVisualLoggerDatabase::Get().GetRowByName(RowName);
 		const TArray<FVisualLogDevice::FVisualLogEntryItem>& Entries = DataRow.GetItems();
 
 		const bool bJustIgnore = Settings->bIgnoreTrivialLogs && Entries.Num() <= Settings->TrivialLogsThreshold;
 		const bool bVisibleByOwnerClasses = FVisualLoggerFilters::Get().MatchObjectName(OwnerClassName.ToString());
-		const bool bIsCollapsed = bJustIgnore || DataRow.GetNumberOfHiddenItems() == Entries.Num() || (SearchFilter.Len() > 0 && OwnerName.ToString().Find(SearchFilter) == INDEX_NONE);
-
+		const bool bHiddenBySearchFilter = (SearchFilter.Len() > 0 && OwnerName.ToString().Find(SearchFilter) == INDEX_NONE);
+		const bool bIsCollapsed = bJustIgnore || DataRow.GetNumberOfHiddenItems() == Entries.Num() || bHiddenBySearchFilter;
+		DataRow.SetHiddenBySearchFilter(bHiddenBySearchFilter);
 		const bool bIsHidden = bIsCollapsed || !bVisibleByOwnerClasses;
 		SetVisibility(bIsHidden ? EVisibility::Collapsed : EVisibility::Visible);
 		if (bIsCollapsed)
@@ -83,12 +85,12 @@ void SLogVisualizerTimeline::UpdateVisibility()
 			Owner->SetSelectionState(SharedThis(this), false, false);
 		}
 
-		FVisualLoggerDatabase::Get().SetRowVisibility(GetName(), !bIsHidden);
+		FVisualLoggerDatabase::Get().SetRowVisibility(RowName, !bIsHidden);
 	}
 	else
 	{
 		Owner->SetSelectionState(SharedThis(this), false, false);
-		FVisualLoggerDatabase::Get().SetRowVisibility(GetName(), false);
+		FVisualLoggerDatabase::Get().SetRowVisibility(RowName, false);
 	}
 }
 
@@ -121,15 +123,7 @@ const TArray<FVisualLogDevice::FVisualLogEntryItem>& SLogVisualizerTimeline::Get
 
 void SLogVisualizerTimeline::AddEntry(const FVisualLogDevice::FVisualLogEntryItem& Entry) 
 { 
-	//Entries.Add(Entry); 
-
-	ULogVisualizerSettings* Settings = ULogVisualizerSettings::StaticClass()->GetDefaultObject<ULogVisualizerSettings>();
-	
-	FVisualLoggerDBRow& DataRow = FVisualLoggerDatabase::Get().GetRowByName(GetName());
-	const TArray<FVisualLogDevice::FVisualLogEntryItem>& Entries = DataRow.GetItems();
-
 	UpdateVisibility();
-
 }
 
 void SLogVisualizerTimeline::Construct(const FArguments& InArgs, TSharedPtr<FVisualLoggerTimeSliderController> InTimeSliderController, TSharedPtr<SVisualLoggerTimelinesContainer> InContainer, 
@@ -205,14 +199,14 @@ void SLogVisualizerTimeline::Construct(const FArguments& InArgs, TSharedPtr<FVis
 			]
 		];
 
-	ULogVisualizerSettings::StaticClass()->GetDefaultObject<ULogVisualizerSettings>()->OnSettingChanged().AddRaw(this, &SLogVisualizerTimeline::HandleLogVisualizerSettingChanged);
+	GetMutableDefault<ULogVisualizerSettings>()->OnSettingChanged().AddRaw(this, &SLogVisualizerTimeline::HandleLogVisualizerSettingChanged);
 	FVisualLoggerDatabase::Get().GetEvents().OnNewItem.AddRaw(this, &SLogVisualizerTimeline::OnNewItemHandler);
 	FVisualLoggerDatabase::Get().GetEvents().OnRowSelectionChanged.AddRaw(this, &SLogVisualizerTimeline::OnRowSelectionChanged);
 }
 
 SLogVisualizerTimeline::~SLogVisualizerTimeline()
 {
-	ULogVisualizerSettings::StaticClass()->GetDefaultObject<ULogVisualizerSettings>()->OnSettingChanged().RemoveAll(this);
+	GetMutableDefault<ULogVisualizerSettings>()->OnSettingChanged().RemoveAll(this);
 	FVisualLoggerDatabase::Get().GetEvents().OnNewItem.RemoveAll(this);
 	FVisualLoggerDatabase::Get().GetEvents().OnRowSelectionChanged.RemoveAll(this);
 }
diff --git a/Engine/Source/Developer/LogVisualizer/Private/SVisualLoggerTimelineBar.cpp b/Engine/Source/Developer/LogVisualizer/Private/SVisualLoggerTimelineBar.cpp
--- a/Engine/Source/Developer/LogVisualizer/Private/SVisualLoggerTimelineBar.cpp
+++ b/Engine/Source/Developer/LogVisualizer/Private/SVisualLoggerTimelineBar.cpp
@@ -6,6 +6,7 @@
 #include "VisualLoggerDatabase.h"
 #include "LogVisualizerStyle.h"
 #include "LogVisualizerPublic.h"
+#include "LogVisualizerSettings.h"
 #include "VisualLoggerTimeSliderController.h"
 #include "Misc/OutputDeviceHelper.h"
 
@@ -64,8 +65,10 @@ FReply SVisualLoggerTimelineBar::OnMouseMove(const FGeometry& MyGeometry, const
 	const int32 NewItemIndex = DBRow.GetClosestItem(ClosestMouseTime);
 	const auto& Items = DBRow.GetItems();
 
-	if (NewItemIndex != MouseMoveClosestItemIndex)
+	if (NewItemIndex != MouseMoveClosestItemIndex
+		|| bToolTipUpdateRequested)
 	{
+		bToolTipUpdateRequested = false;
 		MouseMoveClosestItemIndex = NewItemIndex;
 
 		FString TooltipBuilder;
@@ -75,17 +78,34 @@ FReply SVisualLoggerTimelineBar::OnMouseMove(const FGeometry& MyGeometry, const
 			
 			TooltipBuilder = FString::Printf(TEXT("Time: %.2lf WorldTime: %.2lf"), CurrentEntry.TimeStamp, CurrentEntry.WorldTimeStamp);
 
+			int32 DebugShapesWithoutTextCount = 0;
 			for (const FVisualLogShapeElement& Shape : CurrentEntry.ElementsToDraw)
 			{
-				if (!Shape.Description.IsEmpty())
+				if (FVisualLoggerFilters::Get().ShouldDisplayCategory(Shape.Category, Shape.Verbosity))
 				{
-					TooltipBuilder += FString::Printf(TEXT("\n(shape) %s: %s"), ::ToString(Shape.Verbosity), *Shape.Description);
+					if (Shape.Description.IsEmpty())
+					{
+						DebugShapesWithoutTextCount++;
+					}
+					else
+					{
+						TooltipBuilder += FString::Printf(TEXT("\n(shape) %s[%s]: %s"), *Shape.Category.ToString(), ::ToString(Shape.Verbosity), *Shape.Description);
+					}
 				}
 			}
 
+			if (DebugShapesWithoutTextCount > 0)
+			{
+				TooltipBuilder += FString::Printf(TEXT("\n%d shape(s) without description"), DebugShapesWithoutTextCount);
+			}
+
+			const bool bSearchInsideLogs = GetDefault<ULogVisualizerSettings>()->bSearchInsideLogs;
 			for (const FVisualLogLine& Line : CurrentEntry.LogLines)
 			{
-				TooltipBuilder += FString::Printf(TEXT("\n(log) %s: %s"), ::ToString(Line.Verbosity), *Line.Line);
+				if (FVisualLoggerFilters::Get().ShouldDisplayLine(Line, bSearchInsideLogs))
+				{
+					TooltipBuilder += FString::Printf(TEXT("\n(log) %s: %s"), ::ToString(Line.Verbosity), *Line.Line);
+				}
 			}
 		}
 
@@ -121,6 +141,7 @@ FReply SVisualLoggerTimelineBar::OnMouseButtonDoubleClick(const FGeometry& MyGeo
 
 SVisualLoggerTimelineBar::~SVisualLoggerTimelineBar()
 {
+	FLogVisualizer::Get().GetEvents().OnFiltersChanged.Remove(OnFiltersChangedDelegateHandle);
 }
 
 void SVisualLoggerTimelineBar::Construct(const FArguments& InArgs, TSharedPtr<FVisualLoggerTimeSliderController> InTimeSliderController, TSharedPtr<SLogVisualizerTimeline> InTimelineOwner)
@@ -131,6 +152,12 @@ void SVisualLoggerTimelineBar::Construct(const FArguments& InArgs, TSharedPtr<FV
 	TRange<double> LocalViewRange = TimeSliderController->GetTimeSliderArgs().ViewRange.Get();
 
 	MouseMoveClosestItemIndex = INDEX_NONE;
+
+	// Listen for changes in filters to force refresh the tooltip text of the element closest to the current mouse position
+	OnFiltersChangedDelegateHandle = FLogVisualizer::Get().GetEvents().OnFiltersChanged.AddSPLambda(this, [&bToolTipUpdateRequested = bToolTipUpdateRequested]
+	{
+		bToolTipUpdateRequested = true;
+	});
 }
 
 FVector2D SVisualLoggerTimelineBar::ComputeDesiredSize(float) const
@@ -196,6 +223,7 @@ int32 SVisualLoggerTimelineBar::OnPaint(const FPaintArgs& Args, const FGeometry&
 		const double StartPos = (Entry.TimeStamp - LocalViewRange.GetLowerBoundValue()) * PixelsPerInput - 2;
 		double EndPos = (Entry.TimeStamp - LocalViewRange.GetLowerBoundValue()) * PixelsPerInput + 2;
 		int32 StartIndex = EntryIndex;
+		const bool bSearchInsideLogs = GetDefault<ULogVisualizerSettings>()->bSearchInsideLogs;
 		for (; StartIndex < Entries.Num(); ++StartIndex)
 		{
 			const FVisualLogEntry& CurrentEntry = Entries[StartIndex].Entry;
@@ -214,12 +242,16 @@ int32 SVisualLoggerTimelineBar::OnPaint(const FPaintArgs& Args, const FGeometry&
 			bool bAddedError = false;
 			for (const FVisualLogLine& CurrentLine : LogLines)
 			{
-				if (CurrentLine.Verbosity <= ELogVerbosity::Error && !bAddedError)
+				if (CurrentLine.Verbosity <= ELogVerbosity::Error
+					&& !bAddedError
+					&& FVisualLoggerFilters::Get().ShouldDisplayLine(CurrentLine, bSearchInsideLogs))
 				{
 					ErrorTimes.AddUnique(CurrentEntry.TimeStamp);
 					bAddedError = true;
 				}
-				else if (CurrentLine.Verbosity == ELogVerbosity::Warning && !bAddedWarning)
+				else if (CurrentLine.Verbosity == ELogVerbosity::Warning
+					&& !bAddedWarning
+					&& FVisualLoggerFilters::Get().ShouldDisplayLine(CurrentLine, bSearchInsideLogs))
 				{
 					WarningTimes.AddUnique(CurrentEntry.TimeStamp);
 					bAddedWarning = true;
@@ -255,48 +287,59 @@ int32 SVisualLoggerTimelineBar::OnPaint(const FPaintArgs& Args, const FGeometry&
 		EntryIndex = StartIndex;
 	}
 
-	if (WarningTimes.Num()) RetLayerId++;
-	for (auto CurrentTime : WarningTimes)
+	constexpr double NoSelectionTime = -1;
+	double SelectedTime = NoSelectionTime;
+	if (TimelineOwner.Pin()->IsSelected() && DBRow.GetCurrentItemIndex() != INDEX_NONE)
+	{
+		const  FVisualLogDevice::FVisualLogEntryItem& HighlightedItemEntry = DBRow.GetCurrentItem();
+		SelectedTime = HighlightedItemEntry.Entry.TimeStamp;
+	}
+
+	if (WarningTimes.Num())
 	{
-		double LinePos = (CurrentTime - LocalViewRange.GetLowerBoundValue()) * PixelsPerInput;
+		RetLayerId++;
+	}
 
+	for (const double CurrentTime : WarningTimes)
+	{
+		const double LinePos = (CurrentTime - LocalViewRange.GetLowerBoundValue()) * PixelsPerInput;
+		const float BoxWidth = SelectedTime == CurrentTime ? 10.f : 6.f;
 		FSlateDrawElement::MakeBox(
 			OutDrawElements,
 			RetLayerId,
 			AllottedGeometry.ToPaintGeometry(
-			FVector2f(6.f, AllottedGeometry.GetLocalSize().Y),
-			FSlateLayoutTransform(FVector2f(LinePos - 3.f, 0.0f))),
+			FVector2f(BoxWidth, AllottedGeometry.GetLocalSize().Y),
+			FSlateLayoutTransform(FVector2f(LinePos - (0.5f * BoxWidth), 0.0f))),
 			FillImage,
 			DrawEffects,
 			WarningTimeColor
 			);
 	}
 
-	if (ErrorTimes.Num()) RetLayerId++;
-	for (auto CurrentTime : ErrorTimes)
+	if (ErrorTimes.Num())
 	{
-		double LinePos = (CurrentTime - LocalViewRange.GetLowerBoundValue()) * PixelsPerInput;
+		RetLayerId++;
+	}
 
+	for (const double CurrentTime : ErrorTimes)
+	{
+		const double LinePos = (CurrentTime - LocalViewRange.GetLowerBoundValue()) * PixelsPerInput;
+		const float BoxWidth = SelectedTime == CurrentTime ? 10.f : 6.f;
 		FSlateDrawElement::MakeBox(
 			OutDrawElements,
 			RetLayerId,
 			AllottedGeometry.ToPaintGeometry(
-			FVector2f(6, AllottedGeometry.GetLocalSize().Y),
-			FSlateLayoutTransform(FVector2f(LinePos - 3, 0.0f))),
+			FVector2f(BoxWidth, AllottedGeometry.GetLocalSize().Y),
+			FSlateLayoutTransform(FVector2f(LinePos - (0.5f * BoxWidth), 0.0f))),
 			FillImage,
 			DrawEffects,
 			ErrorTimeColor
 			);
 	}
 
-	int32 BestItemIndex = DBRow.GetClosestItem(LocalScrubPosition);
-
-	if (TimelineOwner.Pin()->IsSelected() && DBRow.GetCurrentItemIndex() != INDEX_NONE)
+	if (SelectedTime != NoSelectionTime)
 	{
-		const auto &HighlightedItemEntry = DBRow.GetCurrentItem();
-		double CurrentTime = HighlightedItemEntry.Entry.TimeStamp;
-		double LinePos = (CurrentTime - LocalViewRange.GetLowerBoundValue()) * PixelsPerInput;
-
+		const double LinePos = (SelectedTime - LocalViewRange.GetLowerBoundValue()) * PixelsPerInput;
 		FSlateDrawElement::MakeBox(
 			OutDrawElements,
 			++RetLayerId,
diff --git a/Engine/Source/Developer/LogVisualizer/Private/SVisualLoggerTimelineBar.h b/Engine/Source/Developer/LogVisualizer/Private/SVisualLoggerTimelineBar.h
--- a/Engine/Source/Developer/LogVisualizer/Private/SVisualLoggerTimelineBar.h
+++ b/Engine/Source/Developer/LogVisualizer/Private/SVisualLoggerTimelineBar.h
@@ -34,10 +34,12 @@ public:
 	virtual ~SVisualLoggerTimelineBar();
 	void Construct(const FArguments& InArgs, TSharedPtr<FVisualLoggerTimeSliderController> InTimeSliderController, TSharedPtr<class SLogVisualizerTimeline> TimelineOwner);
 	virtual int32 OnPaint(const FPaintArgs& Args, const FGeometry& AllottedGeometry, const FSlateRect& MyCullingRect, FSlateWindowElementList& OutDrawElements, int32 LayerId, const FWidgetStyle& InWidgetStyle, bool bParentEnabled) const override;
-	FVector2D ComputeDesiredSize(float) const override;
+	virtual FVector2D ComputeDesiredSize(float) const override;
 
 protected:
 	TSharedPtr<class FVisualLoggerTimeSliderController> TimeSliderController;
 	TWeakPtr<class SLogVisualizerTimeline> TimelineOwner;
-	int32 MouseMoveClosestItemIndex;
+	FDelegateHandle OnFiltersChangedDelegateHandle;
+	int32 MouseMoveClosestItemIndex = INDEX_NONE;
+	bool bToolTipUpdateRequested = false;
 };
diff --git a/Engine/Source/Developer/LogVisualizer/Private/SVisualLoggerView.cpp b/Engine/Source/Developer/LogVisualizer/Private/SVisualLoggerView.cpp
--- a/Engine/Source/Developer/LogVisualizer/Private/SVisualLoggerView.cpp
+++ b/Engine/Source/Developer/LogVisualizer/Private/SVisualLoggerView.cpp
@@ -93,7 +93,7 @@ void SVisualLoggerView::Construct(const FArguments& InArgs, const TSharedRef<FUI
 		.Thickness(FVector2D(6.0f, 6.0f));
 
 
-	ULogVisualizerSettings* Settings = ULogVisualizerSettings::StaticClass()->GetDefaultObject<ULogVisualizerSettings>();
+	const ULogVisualizerSettings* Settings = GetDefault<ULogVisualizerSettings>();
 
 	ChildSlot
 		[
@@ -268,11 +268,14 @@ void SVisualLoggerView::Construct(const FArguments& InArgs, const TSharedRef<FUI
 	});
 
 	OnSearchSplitterResized();
+
+	OnFiltersChangedDelegateHandle = FLogVisualizer::Get().GetEvents().OnFiltersChanged.AddSP(this, &SVisualLoggerView::OnFiltersChanged);
 }
 
 SVisualLoggerView::~SVisualLoggerView()
 {
 	FLogVisualizer::Get().GetEvents().GetAnimationOutlinerFillPercentageFunc.Unbind();
+	FLogVisualizer::Get().GetEvents().OnFiltersChanged.Remove(OnFiltersChangedDelegateHandle);
 }
 
 void SVisualLoggerView::SetAnimationOutlinerFillPercentage(float FillPercentage) 
@@ -450,7 +453,7 @@ TSharedRef<SWidget> SVisualLoggerView::MakeClassesFilterMenu()
 
 void SVisualLoggerView::OnChangedClassesFilter()
 {
-	ULogVisualizerSettings::StaticClass()->GetDefaultObject<ULogVisualizerSettings>()->SaveConfig();
+	GetMutableDefault<ULogVisualizerSettings>()->SaveConfig();
 	for (auto CurrentItem : TimelinesContainer->GetAllNodes())
 	{
 		CurrentItem->UpdateVisibility();
diff --git a/Engine/Source/Developer/LogVisualizer/Private/SVisualLoggerView.h b/Engine/Source/Developer/LogVisualizer/Private/SVisualLoggerView.h
--- a/Engine/Source/Developer/LogVisualizer/Private/SVisualLoggerView.h
+++ b/Engine/Source/Developer/LogVisualizer/Private/SVisualLoggerView.h
@@ -59,5 +59,7 @@ protected:
 	TSharedPtr<class SSearchBox> SearchBox;
 	TSharedPtr<class SComboButton> ClassesComboButton;
 
+	FDelegateHandle OnFiltersChangedDelegateHandle;
+
 	float AnimationOutlinerFillPercentage;
 };
diff --git a/Engine/Source/Developer/LogVisualizer/Private/VisualLoggerDatabase.cpp b/Engine/Source/Developer/LogVisualizer/Private/VisualLoggerDatabase.cpp
--- a/Engine/Source/Developer/LogVisualizer/Private/VisualLoggerDatabase.cpp
+++ b/Engine/Source/Developer/LogVisualizer/Private/VisualLoggerDatabase.cpp
@@ -481,5 +481,5 @@ AActor* FVisualLoggerEditorInterface::GetHelperActor(UWorld* InWorld) const
 
 bool FVisualLoggerEditorInterface::MatchCategoryFilters(const FString& String, ELogVerbosity::Type Verbosity)
 {
-	return FVisualLoggerFilters::Get().MatchCategoryFilters(String, Verbosity);
+	return FVisualLoggerFilters::Get().ShouldDisplayCategoryByString(String, Verbosity);
 }
diff --git a/Engine/Source/Developer/LogVisualizer/Private/VisualLoggerDatabase.h b/Engine/Source/Developer/LogVisualizer/Private/VisualLoggerDatabase.h
--- a/Engine/Source/Developer/LogVisualizer/Private/VisualLoggerDatabase.h
+++ b/Engine/Source/Developer/LogVisualizer/Private/VisualLoggerDatabase.h
@@ -53,9 +53,19 @@ public:
 	const FVisualLogDevice::FVisualLogEntryItem& GetCurrentItem() const;
 	int32 GetCurrentItemIndex() const { return CurrentItemIndex; }
 
+	bool IsRowVisible() const
+	{
+		return !bHiddenBySearchFilter
+			&& GetNumberOfHiddenItems() != GetItems().Num();
+	}
 	void SetItemVisibility(int32 ItemIndex, bool IsVisible);
 	bool IsItemVisible(int32 ItemIndex) const { return HiddenItems[ItemIndex] == false; }
 	int32 GetNumberOfHiddenItems() const { return NumHiddenItems; }
+	void SetHiddenBySearchFilter(const bool bInHiddenBySearchFilter)
+	{
+		bHiddenBySearchFilter = bInHiddenBySearchFilter;
+	}
+
 
 	int32 GetClosestItem(double Time) const;
 	int32 GetClosestItem(double Time, double ScrubTime) const;
@@ -70,6 +80,7 @@ protected:
 	TArray<FVisualLogDevice::FVisualLogEntryItem> Items;
 	TBitArray<> HiddenItems;
 	int32 NumHiddenItems = 0;
+	bool bHiddenBySearchFilter = false;
 };
 
 struct FVisualLoggerDatabase
diff --git a/Engine/Source/Developer/LogVisualizer/Private/VisualLoggerRenderingActor.cpp b/Engine/Source/Developer/LogVisualizer/Private/VisualLoggerRenderingActor.cpp
--- a/Engine/Source/Developer/LogVisualizer/Private/VisualLoggerRenderingActor.cpp
+++ b/Engine/Source/Developer/LogVisualizer/Private/VisualLoggerRenderingActor.cpp
@@ -219,5 +219,5 @@ void AVisualLoggerRenderingActor::IterateDebugShapes(const TFunction<void(const
 
 bool AVisualLoggerRenderingActor::MatchCategoryFilters(const FName& CategoryName, ELogVerbosity::Type Verbosity) const
 {
-	return FVisualLoggerFilters::Get().MatchCategoryFilters(CategoryName.ToString(), Verbosity);
+	return FVisualLoggerFilters::Get().ShouldDisplayCategory(CategoryName, Verbosity);
 }
diff --git a/Engine/Source/Developer/LogVisualizer/Public/LogVisualizerSettings.h b/Engine/Source/Developer/LogVisualizer/Public/LogVisualizerSettings.h
--- a/Engine/Source/Developer/LogVisualizer/Public/LogVisualizerSettings.h
+++ b/Engine/Source/Developer/LogVisualizer/Public/LogVisualizerSettings.h
@@ -5,9 +5,12 @@ GameplayDebuggerSettings.h: Declares the UGameplayDebuggerSettings class.
 =============================================================================*/
 #pragma once
 
+#if UE_ENABLE_INCLUDE_ORDER_DEPRECATED_IN_5_6
 #include "CoreMinimal.h"
+#endif // UE_ENABLE_INCLUDE_ORDER_DEPRECATED_IN_5_6
 #include "UObject/ObjectMacros.h"
 #include "UObject/Object.h"
+#include "VisualLogger/VisualLoggerTypes.h"
 #include "LogVisualizerSettings.generated.h"
 
 struct FVisualLoggerDBRow;
@@ -73,12 +76,50 @@ struct FVisualLoggerFilters : public FVisualLoggerFiltersData
 	void Reset();
 	void InitWith(const FVisualLoggerFiltersData& NewFiltersData);
 
-	/** @return whether given String represents a log category we allow to be displayed at given Verbosity  */
-	bool MatchCategoryFilters(FString String, ELogVerbosity::Type Verbosity = ELogVerbosity::All);
-
-	bool MatchSearchString(FString String) { return SearchBoxFilter.Equals(String, ESearchCase::IgnoreCase); }
-	void SetSearchString(FString InString) { SearchBoxFilter = InString; }
-	FString GetSearchString() { return SearchBoxFilter; }
+	/** @return whether given log line should be displayed */
+	bool ShouldDisplayLine(const FVisualLogLine& Line, const bool bSearchInsideLogs) const;
+
+	/**
+	 * This is the preferred version to determine if a category should be displayed base on its name.
+	 * @return whether given category name represents a log category we allow to be displayed at given Verbosity
+	 */
+	bool ShouldDisplayCategory(FName Name, ELogVerbosity::Type Verbosity = ELogVerbosity::All) const;
+
+	/**
+	 * @return whether given String represents a log category we allow to be displayed at given Verbosity
+	 * @note this function relies on case-insensitive string comparison and is slower than the FName based version
+	 * that should be used when the category name is directly accessible.
+	 * @see ShouldDisplayCategory
+	 */
+	bool ShouldDisplayCategoryByString(FStringView String, ELogVerbosity::Type Verbosity = ELogVerbosity::All) const;
+
+	UE_DEPRECATED(5.6, "Use ShouldDisplayCategoryVerbosity instead")
+	bool MatchCategoryFilters(FString String, const ELogVerbosity::Type Verbosity = ELogVerbosity::All)
+	{
+		return ShouldDisplayCategoryByString(String, Verbosity);
+	}
+
+	/** @return whether given String is a case-insensitive match to the active search filter */
+	bool IsStringMatchingSearchFilter(const FStringView String) const
+	{
+		return SearchBoxFilter.Equals(String.GetData(), ESearchCase::IgnoreCase);
+	}
+
+	UE_DEPRECATED(5.6, "Use IsStringMatchingSearchFilter instead")
+	bool MatchSearchString(FString String)
+	{
+		return IsStringMatchingSearchFilter(String);
+	}
+
+	void SetSearchString(FString InString)
+	{
+		SearchBoxFilter = MoveTemp(InString);
+	}
+
+	FString GetSearchString() const
+	{
+		return SearchBoxFilter;
+	}
 
 	void AddCategory(FString InName, ELogVerbosity::Type InVerbosity);
 	void RemoveCategory(FString InName);
@@ -115,7 +156,10 @@ class LOGVISUALIZER_API ULogVisualizerSettings : public UObject
 
 public:
 	DECLARE_EVENT_OneParam(ULogVisualizerSettings, FSettingChangedEvent, FName /*PropertyName*/);
-	FSettingChangedEvent& OnSettingChanged() { return SettingChangedEvent; }
+	FSettingChangedEvent& OnSettingChanged()
+	{
+		return SettingChangedEvent;
+	}
 
 	/**Whether to show trivial logs, i.e. the ones with only one entry.*/
 	UPROPERTY(EditAnywhere, config, Category = "VisualLogger")
