Fix CIS issue 926233: moved static members to source file

diff --git a/Engine/Source/Developer/LogVisualizer/Internal/SVisualLoggerTableRow.h b/Engine/Source/Developer/LogVisualizer/Internal/SVisualLoggerTableRow.h
--- a/Engine/Source/Developer/LogVisualizer/Internal/SVisualLoggerTableRow.h
+++ b/Engine/Source/Developer/LogVisualizer/Internal/SVisualLoggerTableRow.h
@@ -8,6 +8,8 @@
 #include "Widgets/Views/SListView.h"
 #include "Widgets/Views/STableRow.h"
 
+#define UE_API LOGVISUALIZER_API
+
 class SVisualLoggerTableRow : public SMultiColumnTableRow<TSharedPtr<FLogEntryItem>>
 {
 public:
@@ -17,9 +19,9 @@ public:
 		SLATE_ATTRIBUTE(FText, FilterText)
 	SLATE_END_ARGS()
 
-	static const FLazyName ColumnId_CategoryLabel;
-	static const FLazyName ColumnId_VerbosityLabel;
-	static const FLazyName ColumnId_LogLabel;
+	UE_API static const FLazyName ColumnId_CategoryLabel;
+	UE_API static const FLazyName ColumnId_VerbosityLabel;
+	UE_API static const FLazyName ColumnId_LogLabel;
 
 	void Construct(const FArguments& InArgs, const TSharedRef<STableViewBase>& InOwnerTableView, const TSharedPtr<FLogEntryItem> InEntryItem)
 	{
@@ -68,6 +70,4 @@ private:
 	TSharedPtr<FLogEntryItem> Item;
 };
 
-const FLazyName SVisualLoggerTableRow::ColumnId_CategoryLabel("Category");
-const FLazyName SVisualLoggerTableRow::ColumnId_VerbosityLabel("Verbosity");
-const FLazyName SVisualLoggerTableRow::ColumnId_LogLabel("Log");
+#undef UE_API
diff --git a/Engine/Source/Developer/LogVisualizer/Private/SVisualLoggerTableRow.cpp b/Engine/Source/Developer/LogVisualizer/Private/SVisualLoggerTableRow.cpp
new file mode 100644
--- /dev/null
+++ b/Engine/Source/Developer/LogVisualizer/Private/SVisualLoggerTableRow.cpp
@@ -0,0 +1,8 @@
+// Copyright Epic Games, Inc. All Rights Reserved.
+#pragma once
+
+#include "SVisualLoggerTableRow.h"
+
+const FLazyName SVisualLoggerTableRow::ColumnId_CategoryLabel("Category");
+const FLazyName SVisualLoggerTableRow::ColumnId_VerbosityLabel("Verbosity");
+const FLazyName SVisualLoggerTableRow::ColumnId_LogLabel("Log");
