Fix for CIS# 929692 (removed pragma once from source file)

diff --git a/Engine/Source/Developer/LogVisualizer/Private/SVisualLoggerTableRow.cpp b/Engine/Source/Developer/LogVisualizer/Private/SVisualLoggerTableRow.cpp
--- a/Engine/Source/Developer/LogVisualizer/Private/SVisualLoggerTableRow.cpp
+++ b/Engine/Source/Developer/LogVisualizer/Private/SVisualLoggerTableRow.cpp
@@ -1,5 +1,4 @@
 // Copyright Epic Games, Inc. All Rights Reserved.
-#pragma once
 
 #include "SVisualLoggerTableRow.h"
 
