[VisualLogger] static analysis fix for unused return value

diff --git a/Engine/Source/Developer/LogVisualizer/Private/SVisualLogger.cpp b/Engine/Source/Developer/LogVisualizer/Private/SVisualLogger.cpp
--- a/Engine/Source/Developer/LogVisualizer/Private/SVisualLogger.cpp
+++ b/Engine/Source/Developer/LogVisualizer/Private/SVisualLogger.cpp
@@ -379,7 +379,8 @@ void SVisualLogger::Construct(const FArguments& InArgs, const TSharedRef<SDockTa
 
 	DrawOnCanvasDelegateHandle = UDebugDrawService::Register(TEXT("VisLog"), FDebugDrawDelegate::CreateRaw(VisualLoggerCanvasRenderer.Get(), &FVisualLoggerCanvasRenderer::DrawOnCanvas));
 
-	Cast<AVisualLoggerRenderingActor>(FVisualLoggerEditorInterface::Get()->GetHelperActor(LastUsedWorld.Get()));
+	// Call GetHelperActor to create the helper actor
+	FVisualLoggerEditorInterface::Get()->GetHelperActor(LastUsedWorld.Get());
 }
 
 void SVisualLogger::OnNewLogEntry(const FVisualLogDevice::FVisualLogEntryItem& Entry)
